[UMG MVVM] Uses new UMG Preview plugin

diff --git a/Engine/Plugins/Editor/ModelViewViewModelPreview/ModelViewViewModelPreview.uplugin b/Engine/Plugins/Editor/ModelViewViewModelPreview/ModelViewViewModelPreview.uplugin
new file mode 100644
--- /dev/null
+++ b/Engine/Plugins/Editor/ModelViewViewModelPreview/ModelViewViewModelPreview.uplugin
@@ -0,0 +1,32 @@
+﻿{
+	"FileVersion": 3,
+	"FriendlyName": "UMG Viewmodel for UMG Preview",
+	"Version": 1,
+	"VersionName": "1.0",
+	"CreatedBy": "Epic Games, Inc.",
+	"CreatedByURL": "https://epicgames.com",
+	"Description": "A plugin to support UMG MVVM within the UMG Widget Preview plugin.",
+	"Category": "UI",
+	"EnabledByDefault": false,
+	"CanContainContent": false,
+	"IsBetaVersion": true,
+	"IsExperimentalVersion": true,
+	"Installed": false,
+	"Modules": [
+		{
+			"Name": "ModelViewViewModelPreview",
+			"Type": "Editor",
+			"LoadingPhase": "Default"
+		}
+	],
+	"Plugins": [
+		{
+			"Name": "ModelViewViewModel",
+			"Enabled": true
+		},
+		{
+			"Name": "UMGWidgetPreview",
+			"Enabled": true
+		}
+	]
+}
diff --git a/Engine/Plugins/Editor/ModelViewViewModelPreview/Source/ModelViewViewModelPreview/ModelViewViewModelPreview.Build.cs b/Engine/Plugins/Editor/ModelViewViewModelPreview/Source/ModelViewViewModelPreview/ModelViewViewModelPreview.Build.cs
new file mode 100644
--- /dev/null
+++ b/Engine/Plugins/Editor/ModelViewViewModelPreview/Source/ModelViewViewModelPreview/ModelViewViewModelPreview.Build.cs
@@ -0,0 +1,30 @@
+// Copyright Epic Games, Inc. All Rights Reserved.
+
+using UnrealBuildTool;
+
+public class ModelViewViewModelPreview : ModuleRules
+{
+	public ModelViewViewModelPreview(ReadOnlyTargetRules Target) : base(Target)
+	{
+		PrivateIncludePathModuleNames.AddRange(
+			new string[]
+			{
+				"UMGWidgetPreview",
+			});
+
+		PrivateDependencyModuleNames.AddRange(
+			new string[]
+			{
+				"AdvancedWidgets",
+				"Core",
+				"CoreUObject",
+				"InputCore",
+				"ModelViewViewModel",
+				"ModelViewViewModelBlueprint",
+				"Projects",
+				"Slate",
+				"SlateCore",
+				"UMGWidgetPreview",
+			});
+	}
+}
diff --git a/Engine/Plugins/Editor/ModelViewViewModelPreview/Source/ModelViewViewModelPreview/Private/MVVMWidgetPreviewExtension.cpp b/Engine/Plugins/Editor/ModelViewViewModelPreview/Source/ModelViewViewModelPreview/Private/MVVMWidgetPreviewExtension.cpp
new file mode 100644
--- /dev/null
+++ b/Engine/Plugins/Editor/ModelViewViewModelPreview/Source/ModelViewViewModelPreview/Private/MVVMWidgetPreviewExtension.cpp
@@ -0,0 +1,87 @@
+// Copyright Epic Games, Inc. All Rights Reserved.
+
+#include "MVVMWidgetPreviewExtension.h"
+
+#include "Framework/Docking/LayoutExtender.h"
+#include "IWidgetPreviewToolkit.h"
+#include "Styling/MVVMWidgetPreviewStyle.h"
+#include "Widgets/SMVVMPreviewSourcePanel.h"
+#include "Widgets/Docking/SDockTab.h"
+
+#define LOCTEXT_NAMESPACE "MVVMWidgetPreviewExtension"
+
+namespace UE::MVVM::Private
+{
+	const FLazyName DetailsTabID(TEXT("Details"));
+	const FLazyName DebugSourcePanelTabID = "WidgetPreview_ViewmodelSourcePanel";
+
+	FName GetDetailsTabID()
+	{
+		return DetailsTabID.Resolve();
+	}
+
+	FName FMVVMWidgetPreviewExtension::GetPreviewSourceTabID()
+	{
+		return DebugSourcePanelTabID.Resolve();
+	}
+
+	void FMVVMWidgetPreviewExtension::Register(IUMGWidgetPreviewModule& InWidgetPreviewModule)
+	{
+		InWidgetPreviewModule.OnRegisterTabsForEditor().AddRaw(this, &FMVVMWidgetPreviewExtension::HandleRegisterPreviewEditorTab);
+	}
+
+	void FMVVMWidgetPreviewExtension::Unregister(IUMGWidgetPreviewModule* InWidgetPreviewModule)
+	{
+		check(InWidgetPreviewModule);
+
+		InWidgetPreviewModule->OnRegisterTabsForEditor().RemoveAll(this);
+	}
+
+	void FMVVMWidgetPreviewExtension::HandleRegisterPreviewEditorTab(
+		const TSharedPtr<UE::UMGWidgetPreview::IWidgetPreviewToolkit>& InPreviewEditor,
+		const TSharedRef<FTabManager>& InTabManager)
+	{
+		static const FName PreviewSourceTabId = GetPreviewSourceTabID();
+
+		const TSharedRef<FWorkspaceItem> AssetEditorTabsCategory = InTabManager->GetLocalWorkspaceMenuRoot();
+
+		InTabManager->RegisterTabSpawner(
+			PreviewSourceTabId,
+			FOnSpawnTab::CreateRaw(this, &FMVVMWidgetPreviewExtension::SpawnTab_PreviewSource, TWeakPtr<UE::UMGWidgetPreview::IWidgetPreviewToolkit>(InPreviewEditor)))
+			.SetDisplayName(NSLOCTEXT("DebugSourcePanel", "ViewmodelTabLabel", "Viewmodels"))
+			.SetIcon(FSlateIcon(FMVVMWidgetPreviewStyle::Get().GetStyleSetName(), "BlueprintView.TabIcon"))
+			.SetTooltipText(NSLOCTEXT("DebugSourcePanel", "Viewmodel_ToolTip", "Show the viewmodels panel"))
+			.SetGroup(AssetEditorTabsCategory);
+
+		if (TSharedPtr<FLayoutExtender> LayoutExtender = InPreviewEditor->GetLayoutExtender();
+			LayoutExtender.IsValid())
+		{
+			FTabManager::FTab PreviewSourceTab(FTabId(PreviewSourceTabId, ETabIdFlags::SaveLayout), ETabState::ClosedTab);
+			LayoutExtender->ExtendLayout(GetDetailsTabID(), ELayoutExtensionPosition::After, PreviewSourceTab);
+		}
+	}
+
+	TSharedRef<SDockTab> FMVVMWidgetPreviewExtension::SpawnTab_PreviewSource(
+		const FSpawnTabArgs& Args,
+		TWeakPtr<UE::UMGWidgetPreview::IWidgetPreviewToolkit> InWeakPreviewEditor)
+	{
+		check(Args.GetTabId().TabType == GetPreviewSourceTabID());
+
+		TSharedRef<SDockTab> DockTab = SNew(SDockTab);
+
+		if (InWeakPreviewEditor.IsValid())
+		{
+			DockTab->SetContent(
+				SNew(UE::MVVM::Private::SPreviewSourcePanel, InWeakPreviewEditor.Pin())
+				.AddMetaData<FTagMetaData>(FTagMetaData(TEXT("PreviewSourcePanel"))));
+		}
+		else
+		{
+			DockTab->SetContent(SNullWidget::NullWidget);
+		}
+
+		return DockTab;
+	}
+}
+
+#undef LOCTEXT_NAMESPACE
diff --git a/Engine/Plugins/Editor/ModelViewViewModelPreview/Source/ModelViewViewModelPreview/Private/MVVMWidgetPreviewExtension.h b/Engine/Plugins/Editor/ModelViewViewModelPreview/Source/ModelViewViewModelPreview/Private/MVVMWidgetPreviewExtension.h
new file mode 100644
--- /dev/null
+++ b/Engine/Plugins/Editor/ModelViewViewModelPreview/Source/ModelViewViewModelPreview/Private/MVVMWidgetPreviewExtension.h
@@ -0,0 +1,32 @@
+// Copyright Epic Games, Inc. All Rights Reserved.
+
+#pragma once
+
+#include "Framework/Docking/TabManager.h"
+#include "IUMGWidgetPreviewModule.h"
+#include "Templates/SharedPointer.h"
+
+namespace UE::UMGWidgetPreview
+{
+	class IWidgetPreviewToolkit;
+}
+
+namespace UE::MVVM::Private
+{
+	class FMVVMWidgetPreviewExtension
+		: public TSharedFromThis<FMVVMWidgetPreviewExtension>
+	{
+	public:
+		static FName GetPreviewSourceTabID();
+
+		/** Register various extensibility points with the WidgetPreviewModule. */
+		void Register(IUMGWidgetPreviewModule& InWidgetPreviewModule);
+
+		void Unregister(IUMGWidgetPreviewModule* InWidgetPreviewModule);
+
+	private:
+		void HandleRegisterPreviewEditorTab(const TSharedPtr<UE::UMGWidgetPreview::IWidgetPreviewToolkit>& InPreviewEditor, const TSharedRef<FTabManager>& InTabManager);
+
+		TSharedRef<SDockTab> SpawnTab_PreviewSource(const FSpawnTabArgs& Args, TWeakPtr<UE::UMGWidgetPreview::IWidgetPreviewToolkit> InWeakPreviewEditor);
+	};
+}
diff --git a/Engine/Plugins/Editor/ModelViewViewModelPreview/Source/ModelViewViewModelPreview/Private/ModelViewViewModelPreviewModule.cpp b/Engine/Plugins/Editor/ModelViewViewModelPreview/Source/ModelViewViewModelPreview/Private/ModelViewViewModelPreviewModule.cpp
new file mode 100644
--- /dev/null
+++ b/Engine/Plugins/Editor/ModelViewViewModelPreview/Source/ModelViewViewModelPreview/Private/ModelViewViewModelPreviewModule.cpp
@@ -0,0 +1,37 @@
+// Copyright Epic Games, Inc. All Rights Reserved.
+
+#include "IUMGWidgetPreviewModule.h"
+#include "Modules/ModuleInterface.h"
+#include "Modules/ModuleManager.h"
+#include "MVVMWidgetPreviewExtension.h"
+
+#define LOCTEXT_NAMESPACE "ModelViewViewModelPreview"
+
+class FMVVMPreviewModule
+	: public IModuleInterface
+{
+public:
+	virtual void StartupModule() override
+	{
+		IUMGWidgetPreviewModule& WidgetPreviewModule = FModuleManager::LoadModuleChecked<IUMGWidgetPreviewModule>(WidgetPreviewModuleName);
+		WidgetPreviewExtension = MakeShared<UE::MVVM::Private::FMVVMWidgetPreviewExtension>();
+		WidgetPreviewExtension->Register(WidgetPreviewModule);
+	}
+
+	virtual void ShutdownModule() override
+	{
+		if (IUMGWidgetPreviewModule* WidgetPreviewModule = FModuleManager::GetModulePtr<IUMGWidgetPreviewModule>(WidgetPreviewModuleName))
+		{
+			WidgetPreviewExtension->Unregister(WidgetPreviewModule);
+		}
+	}
+
+private:
+	const FName WidgetPreviewModuleName = "UMGWidgetPreview";
+
+	TSharedPtr<UE::MVVM::Private::FMVVMWidgetPreviewExtension> WidgetPreviewExtension;
+};
+
+#undef LOCTEXT_NAMESPACE
+
+IMPLEMENT_MODULE(FMVVMPreviewModule, ModelViewViewModelPreview)
diff --git a/Engine/Plugins/Editor/ModelViewViewModelPreview/Source/ModelViewViewModelPreview/Private/Styling/MVVMWidgetPreviewStyle.cpp b/Engine/Plugins/Editor/ModelViewViewModelPreview/Source/ModelViewViewModelPreview/Private/Styling/MVVMWidgetPreviewStyle.cpp
new file mode 100644
--- /dev/null
+++ b/Engine/Plugins/Editor/ModelViewViewModelPreview/Source/ModelViewViewModelPreview/Private/Styling/MVVMWidgetPreviewStyle.cpp
@@ -0,0 +1,43 @@
+// Copyright Epic Games, Inc. All Rights Reserved.
+
+#include "MVVMWidgetPreviewStyle.h"
+
+#include "Interfaces/IPluginManager.h"
+#include "Styling/AppStyle.h"
+#include "Styling/CoreStyle.h"
+#include "Styling/SlateStyleMacros.h"
+#include "Styling/SlateStyleRegistry.h"
+#include "Styling/SlateTypes.h"
+
+namespace UE::MVVM::Private
+{
+	FName FMVVMWidgetPreviewStyle::StyleName("MVVMWidgetPreview");
+
+	FMVVMWidgetPreviewStyle& FMVVMWidgetPreviewStyle::Get()
+	{
+		static FMVVMWidgetPreviewStyle Instance;
+		return Instance;
+	}
+
+	FMVVMWidgetPreviewStyle::FMVVMWidgetPreviewStyle()
+		: FSlateStyleSet(StyleName)
+	{
+		const FVector2f Icon16(16.0f);
+		const FVector2f Icon64(64.0f);
+
+		TSharedPtr<IPlugin> Plugin = IPluginManager::Get().FindPlugin(UE_PLUGIN_NAME);
+		check(Plugin.IsValid());
+
+		SetContentRoot(Plugin->GetContentDir() / TEXT("Editor"));
+		SetCoreContentRoot(FPaths::EngineContentDir() / TEXT("Editor/Slate"));
+
+		Set("BlueprintView.TabIcon", new IMAGE_BRUSH_SVG("Slate/ViewModel", Icon16));
+
+		FSlateStyleRegistry::RegisterSlateStyle(*this);
+	}
+
+	FMVVMWidgetPreviewStyle::~FMVVMWidgetPreviewStyle()
+	{
+		FSlateStyleRegistry::UnRegisterSlateStyle(*this);
+	}
+}
diff --git a/Engine/Plugins/Editor/ModelViewViewModelPreview/Source/ModelViewViewModelPreview/Private/Styling/MVVMWidgetPreviewStyle.h b/Engine/Plugins/Editor/ModelViewViewModelPreview/Source/ModelViewViewModelPreview/Private/Styling/MVVMWidgetPreviewStyle.h
new file mode 100644
--- /dev/null
+++ b/Engine/Plugins/Editor/ModelViewViewModelPreview/Source/ModelViewViewModelPreview/Private/Styling/MVVMWidgetPreviewStyle.h
@@ -0,0 +1,21 @@
+// Copyright Epic Games, Inc. All Rights Reserved.
+
+#pragma once
+
+#include "Styling/SlateStyle.h"
+
+namespace UE::MVVM::Private
+{
+	class FMVVMWidgetPreviewStyle
+		: public FSlateStyleSet
+	{
+	public:
+		static FName StyleName;
+
+		static FMVVMWidgetPreviewStyle& Get();
+
+	private:
+		FMVVMWidgetPreviewStyle();
+		virtual ~FMVVMWidgetPreviewStyle() override;
+	};
+}
diff --git a/Engine/Plugins/Editor/ModelViewViewModelPreview/Source/ModelViewViewModelPreview/Private/Widgets/SMVVMPreviewSourcePanel.cpp b/Engine/Plugins/Editor/ModelViewViewModelPreview/Source/ModelViewViewModelPreview/Private/Widgets/SMVVMPreviewSourcePanel.cpp
new file mode 100644
--- /dev/null
+++ b/Engine/Plugins/Editor/ModelViewViewModelPreview/Source/ModelViewViewModelPreview/Private/Widgets/SMVVMPreviewSourcePanel.cpp
@@ -0,0 +1,231 @@
+// Copyright Epic Games, Inc. All Rights Reserved.
+
+#include "Widgets/SMVVMPreviewSourcePanel.h"
+
+#include "IWidgetPreviewToolkit.h"
+#include "MVVMSubsystem.h"
+#include "Styling/AppStyle.h"
+#include "View/MVVMView.h"
+#include "View/MVVMViewClass.h"
+#include "WidgetPreview.h"
+#include "Widgets/Layout/SBorder.h"
+#include "Widgets/PropertyViewer/SFieldIcon.h"
+#include "Widgets/SBoxPanel.h"
+#include "Widgets/Text/STextBlock.h"
+#include "Widgets/Views/SListView.h"
+
+#define LOCTEXT_NAMESPACE "SMVVMDebugSourcePanel"
+
+namespace UE::MVVM::Private
+{
+	class SPreviewSourceView final : public SListView<TSharedPtr<SPreviewSourceEntry>>
+	{
+	};
+
+	class SPreviewSourceEntry
+	{
+	public:
+		SPreviewSourceEntry(UObject* InInstance, FName InName)
+			: WeakInstance(InInstance)
+			, Name(InName)
+		{}
+
+		UClass* GetClass() const
+		{
+			UObject* Instance = WeakInstance.Get();
+			return Instance ? Instance->GetClass() : nullptr;
+		}
+
+		FText GetDisplayName() const
+		{
+			return FText::FromName(Name);
+		}
+
+		UObject* GetInstance() const
+		{
+			return WeakInstance.Get();
+		}
+
+	private:
+		TWeakObjectPtr<UObject> WeakInstance;
+		FName Name;
+	};
+
+	void SPreviewSourcePanel::Construct(const FArguments& InArgs, TSharedPtr<UMGWidgetPreview::IWidgetPreviewToolkit> PreviewEditor)
+	{
+		WeakPreviewEditor = PreviewEditor;
+		check(PreviewEditor);
+
+		if (UWidgetPreview* Preview = PreviewEditor->GetPreview())
+		{
+			HandlePreviewWidgetChanged(EWidgetPreviewWidgetChangeType::Assignment);
+			OnWidgetChangedHandle = Preview->OnWidgetChanged().AddSP(this, &SPreviewSourcePanel::HandlePreviewWidgetChanged);
+	#if UE_WITH_MVVM_DEBUGGING
+			FDebugging::OnViewSourceValueChanged.AddSP(this, &SPreviewSourcePanel::HandleViewChanged);
+	#endif
+		}
+
+		OnSelectedObjectsChangedHandle = PreviewEditor->OnSelectedObjectsChanged().AddSP(this, &SPreviewSourcePanel::HandleSelectedObjectChanged);
+
+		ChildSlot
+		[
+			SNew(SBorder)
+			.BorderImage(FAppStyle::GetBrush("ToolPanel.GroupBorder"))
+			.Padding(FMargin(4.f))
+			[
+				SAssignNew(SourceListView, Private::SPreviewSourceView)
+				.ListItemsSource(&SourceList)
+				.SelectionMode(ESelectionMode::Single)
+				.OnGenerateRow(this, &SPreviewSourcePanel::GenerateWidget)
+				.OnSelectionChanged(this, &SPreviewSourcePanel::HandleSourceSelectionChanged)
+			]
+		];
+
+		HandlePreviewWidgetChanged(EWidgetPreviewWidgetChangeType::Reinstanced);
+	}
+
+
+	SPreviewSourcePanel::~SPreviewSourcePanel()
+	{
+		if (TSharedPtr<UMGWidgetPreview::IWidgetPreviewToolkit> PreviewEditor = WeakPreviewEditor.Pin())
+		{
+			if (UWidgetPreview* Preview = PreviewEditor->GetPreview())
+			{
+				Preview->OnWidgetChanged().Remove(OnWidgetChangedHandle);
+			}
+
+			PreviewEditor->OnSelectedObjectsChanged().Remove(OnSelectedObjectsChangedHandle);
+		}
+	}
+
+
+	void SPreviewSourcePanel::HandlePreviewWidgetChanged(const EWidgetPreviewWidgetChangeType InChangeType)
+	{
+		SourceList.Reset();
+		WeakView.Reset();
+
+		// Only respond to instantiation changes
+		// if (InChangeType == EWidgetPreviewWidgetChangeType::Reinstanced)
+		{
+			if (TSharedPtr<UMGWidgetPreview::IWidgetPreviewToolkit> PreviewEditor = WeakPreviewEditor.Pin())
+			{
+				if (const UUserWidget* NewWidget = PreviewEditor->GetPreview() ? PreviewEditor->GetPreview()->GetWidgetInstance() : nullptr)
+				{
+					if (UMVVMView* View = UMVVMSubsystem::GetViewFromUserWidget(NewWidget))
+					{
+						WeakView = View;
+						for (const FMVVMView_Source& ViewSource : View->GetSources())
+						{
+							FName SourceName = View->GetViewClass()->GetSource(ViewSource.ClassKey).GetName();
+							SourceList.Emplace(MakeShared<Private::SPreviewSourceEntry>(ViewSource.Source, SourceName));
+						}
+					}
+				}
+			}
+		}
+
+		if (SourceListView)
+		{
+			SourceListView->RequestListRefresh();
+		}
+	}
+
+
+	void SPreviewSourcePanel::HandleSelectedObjectChanged(const TConstArrayView<TWeakObjectPtr<UObject>> InSelectedObjects)
+	{
+		if (bIsSelectingListItem)
+		{
+			return;
+		}
+
+		TGuardValue<bool> TmpInternalSelection(bIsSelectingListItem, true);
+		if (TSharedPtr<UMGWidgetPreview::IWidgetPreviewToolkit> PreviewEditor = WeakPreviewEditor.Pin())
+		{
+			if (InSelectedObjects.Num() == 1)
+			{
+				TSharedPtr<Private::SPreviewSourceEntry>* Found = SourceList.FindByPredicate(
+					[ToFind = InSelectedObjects[0]](const TSharedPtr<Private::SPreviewSourceEntry>& Other)
+					{
+						return Other->GetInstance() == ToFind;
+					});
+
+				if (Found)
+				{
+					SourceListView->SetSelection(*Found);
+				}
+				else
+				{
+					SourceListView->SetSelection(nullptr);
+				}
+			}
+			else
+			{
+				SourceListView->SetSelection(nullptr);
+			}
+		}
+	}
+
+
+	void SPreviewSourcePanel::HandleSourceSelectionChanged(TSharedPtr<Private::SPreviewSourceEntry> Entry, ESelectInfo::Type SelectionType) const
+	{
+		if (bIsSelectingListItem)
+		{
+			return;
+		}
+
+		TGuardValue<bool> TmpInternalSelection(bIsSelectingListItem, true);
+		if (TSharedPtr<UMGWidgetPreview::IWidgetPreviewToolkit> PreviewEditor = WeakPreviewEditor.Pin())
+		{
+			if (Entry && Entry->GetInstance())
+			{
+				PreviewEditor->SetSelectedObjects({ Entry->GetInstance() });
+			}
+			else // Nothing selected
+			{
+				// Effectively clears the selection
+				PreviewEditor->SetSelectedObjects({ });
+			}
+		}
+	}
+
+	#if UE_WITH_MVVM_DEBUGGING
+	void SPreviewSourcePanel::HandleViewChanged(const FDebugging::FView& View, const FDebugging::FViewSourceValueArgs& Args)
+	{
+		if (SourceListView)
+		{
+			if (View.GetView() == WeakView.Get())
+			{
+				SourceListView->RebuildList(); // to prevent access to invalid class, rebuild everything.
+			}
+		}
+	}
+	#endif
+
+
+	TSharedRef<ITableRow> SPreviewSourcePanel::GenerateWidget(TSharedPtr<Private::SPreviewSourceEntry> Entry, const TSharedRef<STableViewBase>& OwnerTable) const
+	{
+		typedef STableRow<TSharedPtr<Private::SPreviewSourceEntry>> RowType;
+
+		const TSharedRef<SWidget> FieldIcon = Entry->GetClass() ? SNew(UE::PropertyViewer::SFieldIcon, Entry->GetClass()) : SNullWidget::NullWidget;
+
+		TSharedRef<RowType> NewRow = SNew(RowType, OwnerTable);
+		NewRow->SetContent(SNew(SHorizontalBox)
+			+ SHorizontalBox::Slot()
+			.AutoWidth()
+			.HAlign(HAlign_Right)
+			.VAlign(VAlign_Center)
+			[
+				FieldIcon
+			]
+			+ SHorizontalBox::Slot()
+			.Padding(4.0f)
+			[
+				SNew(STextBlock)
+				.Text(Entry->GetDisplayName())
+			]);
+
+		return NewRow;
+	}
+}
+
+#undef LOCTEXT_NAMESPACE
diff --git a/Engine/Plugins/Editor/ModelViewViewModelPreview/Source/ModelViewViewModelPreview/Private/Widgets/SMVVMPreviewSourcePanel.h b/Engine/Plugins/Editor/ModelViewViewModelPreview/Source/ModelViewViewModelPreview/Private/Widgets/SMVVMPreviewSourcePanel.h
new file mode 100644
--- /dev/null
+++ b/Engine/Plugins/Editor/ModelViewViewModelPreview/Source/ModelViewViewModelPreview/Private/Widgets/SMVVMPreviewSourcePanel.h
@@ -0,0 +1,58 @@
+// Copyright Epic Games, Inc. All Rights Reserved.
+
+#pragma once
+
+#include "Debugging/MVVMDebugging.h"
+#include "WidgetPreview.h"
+#include "Widgets/SCompoundWidget.h"
+
+namespace UE::UMGWidgetPreview
+{
+	class IWidgetPreviewToolkit;
+}
+
+class STableViewBase;
+class ITableRow;
+class UMVVMView;
+class UUserWidget;
+
+#ifndef UE_WITH_MVVM_DEBUGGING
+  #error "MVVM Debugging needs to be enabled."
+#endif
+
+namespace UE::MVVM::Private { class SPreviewSourceView; }
+namespace UE::MVVM::Private { class SPreviewSourceEntry; }
+
+namespace UE::MVVM::Private
+{
+	/** */
+	class SPreviewSourcePanel : public SCompoundWidget
+	{
+	public:
+		SLATE_BEGIN_ARGS(SPreviewSourcePanel) {}
+		SLATE_END_ARGS()
+
+		void Construct(const FArguments& InArgs, TSharedPtr<UMGWidgetPreview::IWidgetPreviewToolkit> PreviewEditor);
+
+		virtual ~SPreviewSourcePanel() override;
+
+	private:
+		void HandlePreviewWidgetChanged(const EWidgetPreviewWidgetChangeType InChangeType);
+		void HandleSelectedObjectChanged(const TConstArrayView<TWeakObjectPtr<UObject>> InSelectedObjects);
+		void HandleSourceSelectionChanged(TSharedPtr<Private::SPreviewSourceEntry> Entry, ESelectInfo::Type SelectionType) const;
+	#if UE_WITH_MVVM_DEBUGGING
+		void HandleViewChanged(const FDebugging::FView&, const FDebugging::FViewSourceValueArgs&);
+	#endif
+		TSharedRef<ITableRow> GenerateWidget(TSharedPtr<Private::SPreviewSourceEntry> Entry, const TSharedRef<STableViewBase>& OwnerTable) const;
+
+		TWeakObjectPtr<UMVVMView> WeakView;
+		TWeakPtr<UMGWidgetPreview::IWidgetPreviewToolkit> WeakPreviewEditor;
+		TSharedPtr<Private::SPreviewSourceView> SourceListView;
+		TArray<TSharedPtr<Private::SPreviewSourceEntry>> SourceList;
+		mutable bool bIsSelectingListItem = false;
+
+		FDelegateHandle OnWidgetChangedHandle;
+		FDelegateHandle OnSelectedObjectsChangedHandle;
+	};
+}
+
diff --git a/Engine/Plugins/Runtime/ModelViewViewModel/Source/ModelViewViewModelEditor/Private/ModelViewViewModelEditorModule.cpp b/Engine/Plugins/Runtime/ModelViewViewModel/Source/ModelViewViewModelEditor/Private/ModelViewViewModelEditorModule.cpp
--- a/Engine/Plugins/Runtime/ModelViewViewModel/Source/ModelViewViewModelEditor/Private/ModelViewViewModelEditorModule.cpp
+++ b/Engine/Plugins/Runtime/ModelViewViewModel/Source/ModelViewViewModelEditor/Private/ModelViewViewModelEditorModule.cpp
@@ -12,7 +12,6 @@
 #include "Customizations/MVVMListViewBaseExtensionCustomizationExtender.h"
 #include "Customizations/MVVMPanelWidgetExtensionCustomizationExtender.h"
 #include "Customizations/MVVMPropertyBindingExtension.h"
-#include "Extensions/MVVMBlueprintViewExtension.h"
 #include "Framework/Application/SlateApplication.h"
 #include "Framework/Docking/LayoutExtender.h"
 #include "Kismet2/BlueprintEditorUtils.h"
@@ -25,7 +24,6 @@
 #include "PropertyEditorModule.h"
 #include "Styling/MVVMEditorStyle.h"
 #include "Tabs/MVVMBindingSummoner.h"
-#include "Tabs/MVVMPreviewSourceSummoner.h"
 #include "Tabs/MVVMViewModelSummoner.h"
 #include "ToolMenus.h"
 #include "UMGEditorModule.h"
@@ -149,10 +147,6 @@ void FModelViewViewModelEditorModule::HandleRegisterBlueprintEditorTab(const FWi
 			}
 		}
 	}
-	else if (ApplicationMode.GetModeName() == FWidgetBlueprintApplicationModes::PreviewMode)
-	{
-		TabFactories.RegisterFactory(MakeShared<UE::MVVM::FPreviewSourceSummoner>(ApplicationMode.GetBlueprintEditor()));
-	}
 }
 
 
diff --git a/Engine/Plugins/Runtime/ModelViewViewModel/Source/ModelViewViewModelEditor/Private/Tabs/MVVMPreviewSourceSummoner.cpp b/Engine/Plugins/Runtime/ModelViewViewModel/Source/ModelViewViewModelEditor/Private/Tabs/MVVMPreviewSourceSummoner.cpp
deleted file mode 100644
diff --git a/Engine/Plugins/Runtime/ModelViewViewModel/Source/ModelViewViewModelEditor/Private/Tabs/MVVMPreviewSourceSummoner.h b/Engine/Plugins/Runtime/ModelViewViewModel/Source/ModelViewViewModelEditor/Private/Tabs/MVVMPreviewSourceSummoner.h
deleted file mode 100644
diff --git a/Engine/Plugins/Runtime/ModelViewViewModel/Source/ModelViewViewModelEditor/Private/Widgets/SMVVMPreviewSourcePanel.cpp b/Engine/Plugins/Runtime/ModelViewViewModel/Source/ModelViewViewModelEditor/Private/Widgets/SMVVMPreviewSourcePanel.cpp
deleted file mode 100644
diff --git a/Engine/Plugins/Runtime/ModelViewViewModel/Source/ModelViewViewModelEditor/Private/Widgets/SMVVMPreviewSourcePanel.h b/Engine/Plugins/Runtime/ModelViewViewModel/Source/ModelViewViewModelEditor/Private/Widgets/SMVVMPreviewSourcePanel.h
deleted file mode 100644
