UMG MVVM: Bindings debugger
- Binding excution events are now logged in the Viewmodel window of the umg widget preview.

#Synced-CL 48846389

diff --git a/Engine/Plugins/Editor/ModelViewViewModelPreview/Source/ModelViewViewModelPreview/ModelViewViewModelPreview.Build.cs b/Engine/Plugins/Editor/ModelViewViewModelPreview/Source/ModelViewViewModelPreview/ModelViewViewModelPreview.Build.cs
--- a/Engine/Plugins/Editor/ModelViewViewModelPreview/Source/ModelViewViewModelPreview/ModelViewViewModelPreview.Build.cs
+++ b/Engine/Plugins/Editor/ModelViewViewModelPreview/Source/ModelViewViewModelPreview/ModelViewViewModelPreview.Build.cs
@@ -16,14 +16,18 @@ public class ModelViewViewModelPreview : ModuleRules
 			new string[]
 			{
 				"AdvancedWidgets",
+				"ApplicationCore",
 				"Core",
 				"CoreUObject",
+				"FieldNotification",
 				"InputCore",
+				"MessageLog",
 				"ModelViewViewModel",
 				"ModelViewViewModelBlueprint",
 				"Projects",
 				"Slate",
 				"SlateCore",
+				"UMG",
 				"UMGWidgetPreview",
 			});
 	}
diff --git a/Engine/Plugins/Editor/ModelViewViewModelPreview/Source/ModelViewViewModelPreview/Private/MVVMWidgetPreviewExtension.cpp b/Engine/Plugins/Editor/ModelViewViewModelPreview/Source/ModelViewViewModelPreview/Private/MVVMWidgetPreviewExtension.cpp
--- a/Engine/Plugins/Editor/ModelViewViewModelPreview/Source/ModelViewViewModelPreview/Private/MVVMWidgetPreviewExtension.cpp
+++ b/Engine/Plugins/Editor/ModelViewViewModelPreview/Source/ModelViewViewModelPreview/Private/MVVMWidgetPreviewExtension.cpp
@@ -5,7 +5,7 @@
 #include "Framework/Docking/LayoutExtender.h"
 #include "IWidgetPreviewToolkit.h"
 #include "Styling/MVVMWidgetPreviewStyle.h"
-#include "Widgets/SMVVMPreviewSourcePanel.h"
+#include "Widgets/SMVVMPreviewViewmodels.h"
 #include "Widgets/Docking/SDockTab.h"
 
 #define LOCTEXT_NAMESPACE "MVVMWidgetPreviewExtension"
@@ -71,9 +71,7 @@ namespace UE::MVVM::Private
 
 		if (InWeakPreviewEditor.IsValid())
 		{
-			DockTab->SetContent(
-				SNew(UE::MVVM::Private::SPreviewSourcePanel, InWeakPreviewEditor.Pin())
-				.AddMetaData<FTagMetaData>(FTagMetaData(TEXT("PreviewSourcePanel"))));
+			DockTab->SetContent(SNew(UE::MVVM::Private::SMVVMPreviewViewmodels, InWeakPreviewEditor.Pin()));
 		}
 		else
 		{
diff --git a/Engine/Plugins/Editor/ModelViewViewModelPreview/Source/ModelViewViewModelPreview/Private/ModelViewViewModelPreviewModule.cpp b/Engine/Plugins/Editor/ModelViewViewModelPreview/Source/ModelViewViewModelPreview/Private/ModelViewViewModelPreviewModule.cpp
--- a/Engine/Plugins/Editor/ModelViewViewModelPreview/Source/ModelViewViewModelPreview/Private/ModelViewViewModelPreviewModule.cpp
+++ b/Engine/Plugins/Editor/ModelViewViewModelPreview/Source/ModelViewViewModelPreview/Private/ModelViewViewModelPreviewModule.cpp
@@ -1,36 +1,48 @@
 // Copyright Epic Games, Inc. All Rights Reserved.
 
+#include "ModelViewViewModelPreviewModule.h"
 #include "IUMGWidgetPreviewModule.h"
-#include "Modules/ModuleInterface.h"
+#include "MessageLogModule.h"
 #include "Modules/ModuleManager.h"
 #include "MVVMWidgetPreviewExtension.h"
 
 #define LOCTEXT_NAMESPACE "ModelViewViewModelPreview"
 
-class FMVVMPreviewModule
-	: public IModuleInterface
+namespace UE::UMGWidgetPreview::Private
 {
-public:
-	virtual void StartupModule() override
+	const FName WidgetPreviewModuleName = "UMGWidgetPreview";
+}
+
+const FLazyName FMVVMPreviewModule::BindingMessageLogName = "MVVMPreviewBindingMessageLog";
+
+void FMVVMPreviewModule::StartupModule()
+{
+	IUMGWidgetPreviewModule& WidgetPreviewModule = FModuleManager::LoadModuleChecked<IUMGWidgetPreviewModule>(UE::UMGWidgetPreview::Private::WidgetPreviewModuleName);
+	WidgetPreviewExtension = MakeShared<UE::MVVM::Private::FMVVMWidgetPreviewExtension>();
+	WidgetPreviewExtension->Register(WidgetPreviewModule);
+
+	FMessageLogModule& MessageLogModule = FModuleManager::LoadModuleChecked<FMessageLogModule>("MessageLog");
+	FMessageLogInitializationOptions InitOptions;
+	InitOptions.bShowFilters = true;
+	InitOptions.bShowPages = false;
+	InitOptions.bShowInLogWindow = false;
+	InitOptions.bAllowClear = true;
+	InitOptions.bScrollToBottom = true;
+	MessageLogModule.RegisterLogListing(BindingMessageLogName, LOCTEXT("MVVMBindingLog", "MVVM Binding Log"), InitOptions);
+}
+
+void FMVVMPreviewModule::ShutdownModule()
+{
+	if (IUMGWidgetPreviewModule* WidgetPreviewModule = FModuleManager::GetModulePtr<IUMGWidgetPreviewModule>(UE::UMGWidgetPreview::Private::WidgetPreviewModuleName))
 	{
-		IUMGWidgetPreviewModule& WidgetPreviewModule = FModuleManager::LoadModuleChecked<IUMGWidgetPreviewModule>(WidgetPreviewModuleName);
-		WidgetPreviewExtension = MakeShared<UE::MVVM::Private::FMVVMWidgetPreviewExtension>();
-		WidgetPreviewExtension->Register(WidgetPreviewModule);
+		WidgetPreviewExtension->Unregister(WidgetPreviewModule);
 	}
 
-	virtual void ShutdownModule() override
+	if (FMessageLogModule* MessageLogModule = FModuleManager::GetModulePtr<FMessageLogModule>("MessageLog"))
 	{
-		if (IUMGWidgetPreviewModule* WidgetPreviewModule = FModuleManager::GetModulePtr<IUMGWidgetPreviewModule>(WidgetPreviewModuleName))
-		{
-			WidgetPreviewExtension->Unregister(WidgetPreviewModule);
-		}
+		MessageLogModule->UnregisterLogListing(BindingMessageLogName);
 	}
-
-private:
-	const FName WidgetPreviewModuleName = "UMGWidgetPreview";
-
-	TSharedPtr<UE::MVVM::Private::FMVVMWidgetPreviewExtension> WidgetPreviewExtension;
-};
+}
 
 #undef LOCTEXT_NAMESPACE
 
diff --git a/Engine/Plugins/Editor/ModelViewViewModelPreview/Source/ModelViewViewModelPreview/Private/ModelViewViewModelPreviewModule.h b/Engine/Plugins/Editor/ModelViewViewModelPreview/Source/ModelViewViewModelPreview/Private/ModelViewViewModelPreviewModule.h
new file mode 100644
--- /dev/null
+++ b/Engine/Plugins/Editor/ModelViewViewModelPreview/Source/ModelViewViewModelPreview/Private/ModelViewViewModelPreviewModule.h
@@ -0,0 +1,28 @@
+// Copyright Epic Games, Inc. All Rights Reserved.
+
+#pragma once
+
+#include "Modules/ModuleInterface.h"
+
+#include "Templates/SharedPointer.h"
+#include "UObject/NameTypes.h"
+
+namespace UE::MVVM::Private
+{
+	class FMVVMWidgetPreviewExtension;
+}
+
+class FMVVMPreviewModule : public IModuleInterface
+{
+public:
+	static const FLazyName BindingMessageLogName;
+
+	//~ Begin IModuleInterface
+	virtual void StartupModule() override;
+	virtual void ShutdownModule() override;
+	//~ End IModuleInterface
+
+private:
+	TSharedPtr<UE::MVVM::Private::FMVVMWidgetPreviewExtension> WidgetPreviewExtension;
+};
+
diff --git a/Engine/Plugins/Editor/ModelViewViewModelPreview/Source/ModelViewViewModelPreview/Private/Widgets/SMVVMPreviewBindingLog.cpp b/Engine/Plugins/Editor/ModelViewViewModelPreview/Source/ModelViewViewModelPreview/Private/Widgets/SMVVMPreviewBindingLog.cpp
new file mode 100644
--- /dev/null
+++ b/Engine/Plugins/Editor/ModelViewViewModelPreview/Source/ModelViewViewModelPreview/Private/Widgets/SMVVMPreviewBindingLog.cpp
@@ -0,0 +1,766 @@
+// Copyright Epic Games, Inc. All Rights Reserved.
+
+#include "Widgets/SMVVMPreviewBindingLog.h"
+
+#include "Blueprint/UserWidget.h"
+#include "Containers/UnrealString.h"
+#include "HAL/PlatformApplicationMisc.h"
+#include "HAL/PlatformMemory.h"
+#include "IMessageLogListing.h"
+#include "INotifyFieldValueChanged.h"
+#include "Internationalization/Text.h"
+#include "IWidgetPreviewToolkit.h"
+#include "Logging/TokenizedMessage.h"
+#include "MessageLogModule.h"
+#include "Misc/Optional.h"
+#include "ModelViewViewModelPreviewModule.h"
+#include "Modules/ModuleManager.h"
+#include "MVVMSubsystem.h"
+#include "Types/MVVMFieldContext.h"
+#include "UObject/Object.h"
+#include "View/MVVMView.h"
+#include "View/MVVMViewClass.h"
+#include "WidgetPreview.h"
+#include "Widgets/SBoxPanel.h"
+#include "Widgets/Views/SHeaderRow.h"
+
+#define LOCTEXT_NAMESPACE "SMVVMPreviewBindingLog"
+
+namespace UE::MVVM::Private
+{
+
+namespace BindingLogUtils
+{
+	// Returns all the value sources for a view.
+	TArray<UObject*> GetViewSources(UMVVMView* InView)
+	{
+		if (!InView)
+		{
+			return {};
+		}
+
+		const TArrayView<const FMVVMView_Source> Sources = InView->GetSources();
+
+		TArray<UObject*> NewSources;
+		NewSources.Reserve(Sources.Num());
+
+		for (const FMVVMView_Source& Source : Sources)
+		{
+			if (Source.Source)
+			{
+				NewSources.Add(Source.Source);
+			}
+		}
+
+		return NewSources;
+	}
+	
+	// Returns the FProperty for the return value of the function.
+	const FProperty* GetRuntimeReturnProperty(const UFunction* InFunction)
+	{
+		FProperty* Result = InFunction->GetReturnProperty();
+
+		if (Result == nullptr && InFunction->HasAllFunctionFlags(FUNC_HasOutParms))
+		{
+			for (TFieldIterator<FProperty> It(InFunction); It && (It->PropertyFlags & CPF_Parm); ++It)
+			{
+				if (It->HasAllPropertyFlags(CPF_OutParm) && !It->HasAnyPropertyFlags(CPF_Deprecated | CPF_EditorOnly | CPF_ConstParm | CPF_ReferenceParm))
+				{
+					Result = *It;
+					break;
+				}
+			}
+		}
+
+		return Result;
+	}
+
+	// Interogate the property and add its value as a message token.
+	// Returns an unset optional on failure.
+	TOptional<FText> GetPropertyValue(UObject* SourceObject, FProperty* InProperty)
+	{
+		if (!InProperty)
+		{
+			return {};
+		}
+
+		if (InProperty->IsA<FArrayProperty>())
+		{
+			return LOCTEXT("SourceValueArray", "[Array]");
+		}
+		else if (FObjectProperty* ObjectProperty = CastField<FObjectProperty>(InProperty))
+		{
+			if (UObject* Value = ObjectProperty->GetObjectPropertyValue_InContainer(SourceObject, 0))
+			{
+				return FText::FromString(Value->GetPathName());
+			}
+			else
+			{
+				return LOCTEXT("SourceValueNull", "[None]");
+			}
+		}
+		// By default will create loctext macros.
+		else if (FTextProperty* TextProperty = CastField<FTextProperty>(InProperty))
+		{
+			FText Value;
+			TextProperty->GetValue_InContainer(SourceObject, &Value);
+
+			return Value;
+		}
+		else
+		{
+			FString Value;
+			InProperty->ExportTextItem_InContainer(Value, SourceObject, nullptr, nullptr, PPF_None);
+
+			return FText::FromString(Value);
+		}
+	}
+
+	// Call the function and add its value as a message token.
+	// Returns an unset optional on failure.
+	TOptional<FText> GetFunctionValue(UObject* SourceObject, UFunction* InFunction)
+	{
+		if (!InFunction)
+		{
+			return {};
+		}
+
+		const FProperty* ReturnProperty = GetRuntimeReturnProperty(InFunction);
+
+		if (!ReturnProperty)
+		{
+			return {};
+		}
+
+		// Bound functions must be const/pure and have no parameters, so only a return value
+		// is necessary for the parameters "struct".
+
+		void* DataPtr = FMemory_Alloca_Aligned(ReturnProperty->GetSize(), ReturnProperty->GetMinAlignment());
+		ReturnProperty->InitializeValue(DataPtr);
+
+		SourceObject->ProcessEvent(InFunction, DataPtr);
+
+		FString Value;
+		ReturnProperty->ExportTextItem_Direct(Value, DataPtr, nullptr, nullptr, PPF_None);
+
+		ReturnProperty->DestroyValue(DataPtr);
+
+		return FText::FromString(Value);
+	}
+
+	// Adds the context's value as a token to the message.
+	// Returns an unset optional on failure.
+	TOptional<FText> GetFieldValue(TValueOrError<UE::MVVM::FFieldContext, void>& InFieldContext, bool bInAllowFunction)
+	{
+		if (InFieldContext.HasError())
+		{
+			return {};
+		}
+
+		UObject* SourceObject = InFieldContext.GetValue().GetObjectVariant().GetUObject();
+
+		if (!SourceObject)
+		{
+			return {};
+		}
+
+		if (InFieldContext.GetValue().GetFieldVariant().IsProperty())
+		{
+			return GetPropertyValue(SourceObject, InFieldContext.GetValue().GetFieldVariant().GetProperty());
+		}
+		else if (bInAllowFunction && InFieldContext.GetValue().GetFieldVariant().IsFunction())
+		{
+			return GetFunctionValue(SourceObject, InFieldContext.GetValue().GetFieldVariant().GetFunction());
+		}
+
+		return {};
+	}
+
+	// Finds the first valid source based on the given source mask.
+	UObject* GetFirstValidSourceByMask(const UMVVMView* InView, uint64 InSourceMask)
+	{
+		if (!InView)
+		{
+			return nullptr;
+		}
+
+		const TArrayView<const FMVVMView_Source> Sources = InView->GetSources();
+
+		for (int32 Index = 0; Index < Sources.Num(); ++Index)
+		{
+			if (InSourceMask & (1ULL << Index))
+			{
+				if (Sources[Index].Source)
+				{
+					return Sources[Index].Source;
+				}
+			}
+		}
+
+		return nullptr;
+	}
+
+	// Find the first valid source by name
+	UObject* GetFirstValidSourceByName(const UMVVMView* InView, FName InName)
+	{
+		if (!InView)
+		{
+			return nullptr;
+		}
+
+		const TArrayView<const FMVVMView_Source> Sources = InView->GetSources();
+
+		for (int32 Index = 0; Index < Sources.Num(); ++Index)
+		{
+			if (Sources[Index].Source)
+			{
+				if (Sources[Index].Source->GetFName() == InName)
+				{
+					return Sources[Index].Source;
+				}
+			}
+		}
+
+		return nullptr;
+	}
+}
+
+void SMVVMPreviewBindingLog::Construct(const FArguments& InArgs, TSharedPtr<UMGWidgetPreview::IWidgetPreviewToolkit> InPreviewEditor)
+{
+	FMessageLogModule& MessageLogModule = FModuleManager::LoadModuleChecked<FMessageLogModule>("MessageLog");
+	MessageLogListing = MessageLogModule.GetLogListing(FMVVMPreviewModule::BindingMessageLogName);
+	MessageLogWidget = MessageLogModule.CreateLogListingWidget(MessageLogListing.ToSharedRef());
+
+	ChildSlot
+	[
+		SNew(SVerticalBox)
+
+		+ SVerticalBox::Slot()
+		.AutoHeight()
+		.Padding(0.f, 4.f, 4.f, 4.f)
+		[
+			SNew(SHeaderRow)
+			+ SHeaderRow::Column("Logs")
+			.FillWidth(1.f)
+			.DefaultLabel(LOCTEXT("BindingLog", "Binding Log"))
+		]
+
+		+ SVerticalBox::Slot()
+		.FillHeight(1.f)
+		.Padding(0.f, 0.f, 4.f, 4.f)
+		[
+			MessageLogWidget.ToSharedRef()
+		]
+	];
+
+	SetPreviewEditor(InPreviewEditor);
+
+#if UE_WITH_MVVM_DEBUGGING
+	OnOnLibraryBindingExecutedHandle = FDebugging::OnLibraryBindingExecuted.AddSP(
+		this,
+		&SMVVMPreviewBindingLog::HandleOnLibraryBindingExecuted
+	);
+#endif
+}
+
+SMVVMPreviewBindingLog::~SMVVMPreviewBindingLog()
+{
+	if (UObjectInitialized())
+	{
+		SetPreviewEditor(nullptr);
+	}
+
+#if UE_WITH_MVVM_DEBUGGING
+	if (OnOnLibraryBindingExecutedHandle.IsValid())
+	{
+		FDebugging::OnLibraryBindingExecuted.Remove(OnOnLibraryBindingExecutedHandle);
+		OnOnLibraryBindingExecutedHandle.Reset();
+	}
+#endif
+}
+
+void SMVVMPreviewBindingLog::SetPreviewEditor(TSharedPtr<UMGWidgetPreview::IWidgetPreviewToolkit> InPreviewEditor)
+{
+	TSharedPtr<UMGWidgetPreview::IWidgetPreviewToolkit> CurrentPreviewEditor = WeakPreviewEditor.Pin();
+
+	if (CurrentPreviewEditor.Get() == InPreviewEditor.Get())
+	{
+		return;
+	}
+
+	WeakPreviewEditor = InPreviewEditor;
+
+	if (!InPreviewEditor.IsValid())
+	{
+		SetWidgetPreview(nullptr);
+	}
+	else
+	{
+		if (UWidgetPreview* WidgetPreview = InPreviewEditor->GetPreview())
+		{
+			SetWidgetPreview(WidgetPreview);
+		}
+	}
+}
+
+void SMVVMPreviewBindingLog::SetWidgetPreview(UWidgetPreview* InWidgetPreview)
+{
+	UWidgetPreview* CurrentWidgetPreview = WeakWidgetPreview.Get();
+
+	if (InWidgetPreview == CurrentWidgetPreview)
+	{
+		return;
+	}
+
+	if (OnWidgetChangedHandle.IsValid())
+	{
+		if (CurrentWidgetPreview)
+		{
+			CurrentWidgetPreview->OnWidgetChanged().Remove(OnWidgetChangedHandle);
+			OnWidgetChangedHandle.Reset();
+		}
+	}
+
+	WeakWidgetPreview = InWidgetPreview;
+
+	if (!InWidgetPreview)
+	{
+		SetPreviewWidget(nullptr);
+	}
+	else
+	{
+		OnWidgetChangedHandle = InWidgetPreview->OnWidgetChanged().AddSP(
+			this,
+			&SMVVMPreviewBindingLog::HandleOnWidgetChanged
+		);
+
+		if (UUserWidget* NewPreviewWidget = InWidgetPreview->GetWidgetInstance())
+		{
+			SetPreviewWidget(NewPreviewWidget);
+		}
+	}
+}
+
+void SMVVMPreviewBindingLog::SetPreviewWidget(UUserWidget* InWidget)
+{
+	if (UMVVMView* CurrentView = WeakView.Get())
+	{
+		UUserWidget* CurrentWidget = Cast<UUserWidget>(CurrentView->GetOuter());
+
+		if (CurrentWidget == InWidget)
+		{
+			return;
+		}
+	}
+
+	// Don't store a weak reference to the user widget as we're not binding anything to it.
+
+	if (!InWidget)
+	{
+		SetView(nullptr);
+	}
+	else
+	{
+		SetView(UMVVMSubsystem::GetViewFromUserWidget(InWidget));
+	}	
+}
+
+void SMVVMPreviewBindingLog::SetView(UMVVMView* InView)
+{
+	UMVVMView* CurrentView = WeakView.Get();
+
+	if (InView == CurrentView)
+	{
+		return;
+	}
+
+#if UE_WITH_MVVM_DEBUGGING
+	if (OnViewSourceChangedHandle.IsValid())
+	{
+		FDebugging::OnViewSourceValueChanged.Remove(OnViewSourceChangedHandle);
+		OnViewSourceChangedHandle.Reset();
+	}
+#endif
+
+	WeakView = InView;
+
+	if (!InView)
+	{
+		SetSources({});
+	}
+	else
+	{
+#if UE_WITH_MVVM_DEBUGGING
+		OnViewSourceChangedHandle = FDebugging::OnViewSourceValueChanged.AddSP(
+			this,
+			&SMVVMPreviewBindingLog::HandleOnViewSourceChanged
+		);
+#endif	
+
+		SetSources(BindingLogUtils::GetViewSources(InView));
+	}
+}
+
+void SMVVMPreviewBindingLog::SetSources(const TArray<UObject*>& InSources)
+{
+	for (const TWeakObjectPtr<UObject>& WeakSource : WeakSources)
+	{
+		UObject* Source = WeakSource.Get();
+
+		if (!Source)
+		{
+			continue;
+		}
+
+		INotifyFieldValueChanged* NotifyField = Cast<INotifyFieldValueChanged>(Source);
+
+		if (!NotifyField)
+		{
+			continue;
+		}
+
+		NotifyField->RemoveAllFieldValueChangedDelegates(this);
+	}
+
+	WeakSources.Empty(InSources.Num());
+
+	for (UObject* Source : InSources)
+	{
+		INotifyFieldValueChanged* NotifyField = Cast<INotifyFieldValueChanged>(Source);
+
+		if (!NotifyField)
+		{
+			continue;
+		}
+
+		const UE::FieldNotification::IClassDescriptor& ClassDescriptor = NotifyField->GetFieldNotificationDescriptor();
+
+		ClassDescriptor.ForEachField(
+			Source->GetClass(),
+			[this, NotifyField](UE::FieldNotification::FFieldId InFieldId)
+			{
+				NotifyField->AddFieldValueChangedDelegate(
+					InFieldId,
+					INotifyFieldValueChanged::FFieldValueChangedDelegate::CreateSP(
+						this,
+						&SMVVMPreviewBindingLog::HandleOnFieldValueChanged
+					)
+				);
+
+				return true;
+			}
+		);
+
+		WeakSources.Add(Source);
+	}
+}
+
+
+void SMVVMPreviewBindingLog::HandleOnWidgetChanged(const EWidgetPreviewWidgetChangeType InChangeType)
+{
+	// Make sure that we still have the correct preview widget
+	if (UWidgetPreview* WidgetPreview = WeakWidgetPreview.Get())
+	{
+		SetPreviewWidget(WidgetPreview->GetWidgetInstance());
+	}
+	else
+	{
+		SetPreviewWidget(nullptr);
+	}
+}
+
+#if UE_WITH_MVVM_DEBUGGING
+void SMVVMPreviewBindingLog::HandleOnViewSourceChanged(const FDebugging::FView& InView, const FDebugging::FViewSourceValueArgs& InArgs)
+{
+	SetSources(BindingLogUtils::GetViewSources(WeakView.Get()));
+}
+#endif
+
+void SMVVMPreviewBindingLog::HandleOnFieldValueChanged(UObject* InSourceObject, UE::FieldNotification::FFieldId InFieldId)
+{
+	if (!InSourceObject)
+	{
+		return;
+	}
+
+	UE::FieldNotification::FFieldVariant Field = InFieldId.ToVariant(InSourceObject);
+
+	if (!Field.IsValid())
+	{
+		return;
+	}
+
+	TSharedRef<FTokenizedMessage> SourceLog = MakeLog(
+		EMessageSeverity::Info,
+		LOCTEXT("SourceLog", "Trigger: ")
+	);
+
+	if (UWidget* Widget = Cast<UWidget>(InSourceObject))
+	{
+		AddSourceToken(SourceLog, Widget);
+	}
+	else
+	{
+		AddSourceToken(SourceLog, InSourceObject);
+	}
+
+	if (Field.IsProperty())
+	{
+		SourceLog->AddText(INVTEXT("."));
+		SourceLog->AddText(Field.GetProperty()->GetDisplayNameText());
+	}
+	else if (Field.IsFunction())
+	{
+		SourceLog->AddText(INVTEXT("."));
+		SourceLog->AddText(Field.GetFunction()->GetDisplayNameText());
+		SourceLog->AddText(INVTEXT("()"));
+	}
+	else
+	{
+		ensureAlwaysMsgf(false, TEXT("Invalid field type not caught by FieldVariant.IsValid()"));
+		SourceLog->AddText(INVTEXT(".???"));
+		SourceLog->SetSeverity(EMessageSeverity::Error);
+	}
+
+	using namespace BindingLogUtils;
+
+	TSharedRef<FTokenizedMessage> SourceValueLog = MakeLog(
+		EMessageSeverity::Info,
+		LOCTEXT("SourceValueLog", "Source Value: ")
+	);
+
+	SourceValueLog->SetIndentationLevel(2);
+
+	bool bHasError = false;
+
+	if (FProperty* Property = Field.GetProperty())
+	{
+		if (TOptional<FText> Value = BindingLogUtils::GetPropertyValue(InSourceObject, Property))
+		{
+			AddCopyToken(SourceValueLog, Value.GetValue());
+		}
+		else
+		{
+			SourceValueLog->AddText<FText>(LOCTEXT("SourceValueError_PropValue", "Error Failed to get value for property %s"), FText::FromString(Property->GetName()));
+			bHasError = true;
+		}
+	}
+	else if (UFunction* Function = Field.GetFunction())
+	{
+		if (TOptional<FText> Value = BindingLogUtils::GetFunctionValue(InSourceObject, Function))
+		{
+			AddCopyToken(SourceValueLog, Value.GetValue());
+		}
+		else
+		{
+			SourceValueLog->AddText<FText>(LOCTEXT("SourceValueError_FuncValue", "Error Failed to get value for function %s"), FText::FromString(Function->GetName()));
+			bHasError = true;
+		}
+	}
+	else
+	{
+		SourceValueLog->AddText<FText>(LOCTEXT("SourceValueError_FieldValue", "Error Failed to get value for field %s"), FText::FromName(Field.GetFName()));
+		bHasError = true;
+	}
+}
+
+
+#if UE_WITH_MVVM_DEBUGGING
+void SMVVMPreviewBindingLog::HandleOnLibraryBindingExecuted(const FDebugging::FView& InView, const FDebugging::FLibraryBindingExecutedArgs& InArgs)
+{
+	const UMVVMView* View = InView.GetView();
+
+	if (!View || View != WeakView.Get())
+	{
+		return;
+	}
+
+	const UMVVMViewClass* ViewClass = View->GetViewClass();
+	
+	if (!ViewClass)
+	{
+		return;
+	}
+
+	const TArrayView<const FMVVMViewClass_Binding> Bindings = ViewClass->GetBindings();
+
+	if (!Bindings.IsValidIndex(InArgs.Binding.GetIndex()))
+	{
+		return;
+	}
+
+	const FMVVMVCompiledBinding& Binding = Bindings[InArgs.Binding.GetIndex()].GetBinding();
+
+	UUserWidget* UserWidget = Cast<UUserWidget>(View->GetOuter());
+
+	/*
+	 * Add binding destination
+	 */
+	TValueOrError<FString, FString> DestinationFieldPathString = ViewClass->GetBindingLibrary().FieldPathToString(Binding.GetDestinationFieldPath(), /* Display name */ true);
+	FString DestinationString = DestinationFieldPathString.HasValue() ? DestinationFieldPathString.StealValue() : DestinationFieldPathString.StealError();
+
+	TSharedRef<FTokenizedMessage> DestinationLog = MakeLog(
+		(DestinationFieldPathString.HasError() || DestinationString.IsEmpty())
+			? EMessageSeverity::Error
+			: EMessageSeverity::Info,
+		LOCTEXT("DestinationLog", "Set:")
+	);
+
+	DestinationLog->SetIndentationLevel(1);
+
+	bool bHasError = false;
+	FString DestinationName;
+
+	if (DestinationFieldPathString.HasError())
+	{
+		const FText ErrorFormat = LOCTEXT("DestinationErrorFormat", "(Error: {0})");
+		DestinationName += FText::Format(ErrorFormat, FText::FromString(DestinationString)).ToString();
+	}
+	else if (DestinationString.IsEmpty())
+	{
+		DestinationName += LOCTEXT("DestinationMissing", "???").ToString();
+	}
+	else
+	{
+		if (UserWidget)
+		{
+			AddSourceToken(DestinationLog, UserWidget);
+			DestinationName = TEXT(".");
+		}
+
+		DestinationName += DestinationString;
+	}
+
+	DestinationLog->AddToken(FTextToken::Create(FText::FromString(DestinationName)));
+
+	/*
+	 * Add destination value if there's a conversion function
+	 * The value may/will differ from the source value. Explicitly add a message about the transformed value.
+	 */
+	if (UserWidget && Binding.GetDestinationFieldPath().IsValid() && Binding.GetConversionFunctionFieldPath().IsValid())
+	{
+		TValueOrError<UE::MVVM::FFieldContext, void> DestinationFieldContext = ViewClass->GetBindingLibrary().EvaluateFieldPath(UserWidget, Binding.GetDestinationFieldPath());
+
+		TSharedRef<FTokenizedMessage> DestinationValueLog = MakeLog(
+			DestinationFieldPathString.HasError()
+				? EMessageSeverity::Error
+				: EMessageSeverity::Info,
+			LOCTEXT("DestinationValueLog", "Target Value: ")
+		);
+
+		DestinationValueLog->SetIndentationLevel(3);
+
+		/*
+		 * Prepend conversion information to the value
+		 */
+		if (Binding.GetConversionFunctionFieldPath().IsValid())
+		{
+			FText FormattedConversion;
+
+			TValueOrError<FString, FString> ConversionFieldPathString = ViewClass->GetBindingLibrary().FieldPathToString(Binding.GetConversionFunctionFieldPath(), false);
+			const FString ConversionString = ConversionFieldPathString.HasValue() ? ConversionFieldPathString.StealValue() : ConversionFieldPathString.StealError();
+
+			if (!ConversionString.IsEmpty())
+			{
+				if (ConversionFieldPathString.HasError())
+				{
+					const FText ConversionErrorFormat = LOCTEXT("ConversionErrorFormat", "(Error: {0})");
+					FormattedConversion = FText::Format(ConversionErrorFormat, FText::FromString(ConversionString));
+				}
+				else
+				{
+					const FText ConversionErrorFormat = LOCTEXT("ConversionLogFormat", "({0})");
+					FormattedConversion = FText::Format(ConversionErrorFormat, FText::FromString(ConversionString));
+				}
+			}
+
+			if (!FormattedConversion.IsEmpty())
+			{
+				DestinationValueLog->AddToken(FTextToken::Create(FormattedConversion));
+				DestinationValueLog->AddToken(FTextToken::Create(INVTEXT(" ")));
+			}
+		}
+
+		if (TOptional<FText> Value = BindingLogUtils::GetFieldValue(DestinationFieldContext, /* Allow Function */ false))
+		{
+			AddCopyToken(DestinationValueLog, Value.GetValue());
+		}
+		else
+		{
+			DestinationValueLog->AddText(LOCTEXT("DestinationValueError", "Error"));
+		};
+	}
+}
+#endif
+
+TSharedRef<FTokenizedMessage> SMVVMPreviewBindingLog::MakeLog(EMessageSeverity::Type InSeverity, const FText& InMessage)
+{
+	TSharedRef<FTokenizedMessage> NewMessage = FTokenizedMessage::Create(InSeverity);
+	NewMessage->AddToken(FTextToken::Create(InMessage));
+	MessageLogListing->AddMessage(NewMessage);
+	return NewMessage;
+}
+
+void SMVVMPreviewBindingLog::HandleObjectTokenClicked(TWeakObjectPtr<UObject> InWeakObject)
+{
+	if (TSharedPtr<UMGWidgetPreview::IWidgetPreviewToolkit> PreviewEditor = WeakPreviewEditor.Pin())
+	{
+		if (UObject* ClickedObject = InWeakObject.Get())
+		{
+			PreviewEditor->SetSelectedObjects({ClickedObject});
+		}
+	}
+}
+
+void SMVVMPreviewBindingLog::AddSourceToken(TSharedRef<FTokenizedMessage> InMessage, TNotNull<UObject*> InSource)
+{
+	if (UMVVMView* View = WeakView.Get())
+	{
+		const UObject* ViewOuter = View->GetOuter();
+
+		const FText Path = (InSource == ViewOuter)
+			? FText::FromName(InSource->GetFName())
+			: FText::FromString(InSource->GetPathName(View->GetOuter()));
+
+		InMessage->AddToken(
+			FActionToken::Create(
+				Path,
+				LOCTEXT("SelectSource", "Select Source"),
+				FOnActionTokenExecuted::CreateSP(
+					this,
+					&SMVVMPreviewBindingLog::HandleObjectTokenClicked,
+					TWeakObjectPtr<UObject>(InSource)
+				)
+			)
+		);
+	}
+}
+
+void SMVVMPreviewBindingLog::AddCopyToken(TSharedRef<FTokenizedMessage> InMessage, const FText& InText)
+{
+	InMessage->AddText(InText);
+
+	InMessage->AddText(INVTEXT(" "));
+
+	InMessage->AddToken(
+		FActionToken::Create(
+			LOCTEXT("CopyLink", "Copy"),
+			LOCTEXT("CopyText", "Copy to the Clipboard"),
+			FOnActionTokenExecuted::CreateSP(
+				this,
+				&SMVVMPreviewBindingLog::HandleCopyTokenClicked,
+				InText
+			)
+		)
+	);
+}
+
+void SMVVMPreviewBindingLog::HandleCopyTokenClicked(FText InText)
+{
+	FPlatformApplicationMisc::ClipboardCopy(*InText.ToString());
+}
+
+}
+
+#undef LOCTEXT_NAMESPACE
+
diff --git a/Engine/Plugins/Editor/ModelViewViewModelPreview/Source/ModelViewViewModelPreview/Private/Widgets/SMVVMPreviewBindingLog.h b/Engine/Plugins/Editor/ModelViewViewModelPreview/Source/ModelViewViewModelPreview/Private/Widgets/SMVVMPreviewBindingLog.h
new file mode 100644
--- /dev/null
+++ b/Engine/Plugins/Editor/ModelViewViewModelPreview/Source/ModelViewViewModelPreview/Private/Widgets/SMVVMPreviewBindingLog.h
@@ -0,0 +1,124 @@
+// Copyright Epic Games, Inc. All Rights Reserved.
+
+#pragma once
+
+#include "Widgets/SCompoundWidget.h"
+
+#include "Delegates/IDelegateInstance.h"
+#include "UObject/WeakObjectPtr.h"
+
+#if UE_WITH_MVVM_DEBUGGING
+#include "Debugging/MVVMDebugging.h"
+#endif
+
+class FText;
+class FTokenizedMessage;
+class IMessageLogListing;
+class INotifyFieldValueChanged;
+class UMVVMView;
+class UObject;
+class UUserWidget;
+class UWidgetPreview;
+enum class EWidgetPreviewWidgetChangeType : uint8;
+
+namespace EMessageSeverity
+{
+	enum Type : int;
+}
+
+namespace UE::UMGWidgetPreview
+{
+	class IWidgetPreviewToolkit;
+}
+
+namespace UE::FieldNotification
+{
+	struct FFieldId;
+}
+
+namespace UE::MVVM::Private
+{
+
+class SPreviewSourceView;
+class SPreviewSourceEntry;
+
+/** Log to display the changes in the UMG Preview Widget MVVM Values/Bindings */
+class SMVVMPreviewBindingLog : public SCompoundWidget
+{
+public:
+	SLATE_BEGIN_ARGS(SMVVMPreviewBindingLog) {}
+	SLATE_END_ARGS()
+
+	void Construct(const FArguments& InArgs, TSharedPtr<UMGWidgetPreview::IWidgetPreviewToolkit> InPreviewEditor);
+
+	virtual ~SMVVMPreviewBindingLog() override;
+
+private:
+	TWeakPtr<UMGWidgetPreview::IWidgetPreviewToolkit> WeakPreviewEditor;
+	TWeakObjectPtr<UWidgetPreview> WeakWidgetPreview;
+	TWeakObjectPtr<UMVVMView> WeakView;
+
+	FDelegateHandle OnWidgetChangedHandle;
+
+#if UE_WITH_MVVM_DEBUGGING
+	FDelegateHandle OnOnLibraryBindingExecutedHandle;
+#endif
+
+	// Must implement INotifyFieldValueChanged
+	// Stored here to remember whose events we subscribed to.
+	TArray<TWeakObjectPtr<UObject>> WeakSources;
+
+#if UE_WITH_MVVM_DEBUGGING
+	FDelegateHandle OnViewSourceChangedHandle;
+#endif
+
+	TSharedPtr<IMessageLogListing> MessageLogListing;
+	TSharedPtr<SWidget> MessageLogWidget;
+
+	// Sets and unsets the preview editor, adding and removing bindings.
+	void SetPreviewEditor(TSharedPtr<UMGWidgetPreview::IWidgetPreviewToolkit> InPreviewEditor);
+
+	// Sets and unsets the widget preview object, adding and removing bindings.
+	void SetWidgetPreview(UWidgetPreview* InWidgetPreview);
+
+	// Checks the current view to see if we need to update it with this widget's view.
+	void SetPreviewWidget(UUserWidget* InWidget);
+
+	// Sets and unsets the view.
+	void SetView(UMVVMView* InView);
+
+	// When the widget in the preview is changed, update the view.
+	void HandleOnWidgetChanged(const EWidgetPreviewWidgetChangeType InChangeType);
+
+#if UE_WITH_MVVM_DEBUGGING
+	void HandleOnLibraryBindingExecuted(const FDebugging::FView& InView, const FDebugging::FLibraryBindingExecutedArgs& InArgs);
+#endif
+
+	// Removes previous binds to field notify sources and adds new binds to the provided sources.
+	void SetSources(const TArray<UObject*>& InSources);
+
+#if UE_WITH_MVVM_DEBUGGING
+	// When the sources change, update the source bindings.
+	void HandleOnViewSourceChanged(const FDebugging::FView& InView, const FDebugging::FViewSourceValueArgs& InArgs);
+#endif
+
+	// When a field notify value changes, produce a log entry.
+	void HandleOnFieldValueChanged(UObject* InSourceObject, UE::FieldNotification::FFieldId InFieldId);
+
+	// Creates a log message in the binding log message log and returns it.
+	TSharedRef<FTokenizedMessage> MakeLog(EMessageSeverity::Type InSeverity, const FText& InMessage);
+
+	// Adds an action token that selects the given source object.
+	void AddSourceToken(TSharedRef<FTokenizedMessage> InMessage, TNotNull<UObject*> InSource);
+		
+	// Callback to select an object.
+	void HandleObjectTokenClicked(TWeakObjectPtr<UObject> InWeakObject);
+
+	// Adds an action that copies the given text to the clipboard
+	void AddCopyToken(TSharedRef<FTokenizedMessage> InMessage, const FText& InText);
+
+	// Handles the copy callback to copy text to the clipboard
+	void HandleCopyTokenClicked(FText InText);
+};
+
+}
diff --git a/Engine/Plugins/Editor/ModelViewViewModelPreview/Source/ModelViewViewModelPreview/Private/Widgets/SMVVMPreviewSourcePanel.cpp b/Engine/Plugins/Editor/ModelViewViewModelPreview/Source/ModelViewViewModelPreview/Private/Widgets/SMVVMPreviewSourcePanel.cpp
--- a/Engine/Plugins/Editor/ModelViewViewModelPreview/Source/ModelViewViewModelPreview/Private/Widgets/SMVVMPreviewSourcePanel.cpp
+++ b/Engine/Plugins/Editor/ModelViewViewModelPreview/Source/ModelViewViewModelPreview/Private/Widgets/SMVVMPreviewSourcePanel.cpp
@@ -70,14 +70,20 @@ namespace UE::MVVM::Private
 		ChildSlot
 		[
 			SNew(SBorder)
-			.BorderImage(FAppStyle::GetBrush("ToolPanel.GroupBorder"))
-			.Padding(FMargin(4.f))
+			.BorderImage(&FAppStyle::Get().GetWidgetStyle<FTableViewStyle>("ListView").BackgroundBrush)
+			.Padding(4.f)
 			[
 				SAssignNew(SourceListView, Private::SPreviewSourceView)
 				.ListItemsSource(&SourceList)
 				.SelectionMode(ESelectionMode::Single)
 				.OnGenerateRow(this, &SPreviewSourcePanel::GenerateWidget)
 				.OnSelectionChanged(this, &SPreviewSourcePanel::HandleSourceSelectionChanged)
+				.HeaderRow(
+					SNew(SHeaderRow)
+					+ SHeaderRow::Column("Sources")
+					.FillWidth(1.f)
+					.DefaultLabel(LOCTEXT("Sources", "Sources"))
+				)
 			]
 		];
 
diff --git a/Engine/Plugins/Editor/ModelViewViewModelPreview/Source/ModelViewViewModelPreview/Private/Widgets/SMVVMPreviewViewmodels.cpp b/Engine/Plugins/Editor/ModelViewViewModelPreview/Source/ModelViewViewModelPreview/Private/Widgets/SMVVMPreviewViewmodels.cpp
new file mode 100644
--- /dev/null
+++ b/Engine/Plugins/Editor/ModelViewViewModelPreview/Source/ModelViewViewModelPreview/Private/Widgets/SMVVMPreviewViewmodels.cpp
@@ -0,0 +1,40 @@
+// Copyright Epic Games, Inc. All Rights Reserved.
+
+#include "Widgets/SMVVMPreviewViewmodels.h"
+
+#include "Widgets/Layout/SBox.h"
+#include "Widgets/SBoxPanel.h"
+#include "Widgets/SMVVMPreviewBindingLog.h"
+#include "Widgets/SMVVMPreviewSourcePanel.h"
+
+namespace UE::MVVM::Private
+{
+
+void SMVVMPreviewViewmodels::Construct(const FArguments& InArgs, TSharedPtr<UMGWidgetPreview::IWidgetPreviewToolkit> InPreviewEditor)
+{
+	SourcePanel = SNew(SPreviewSourcePanel, InPreviewEditor)
+		.AddMetaData<FTagMetaData>(FTagMetaData(TEXT("PreviewSourcePanel")));
+
+	BindingLog = SNew(SMVVMPreviewBindingLog, InPreviewEditor)
+		.AddMetaData<FTagMetaData>(FTagMetaData(TEXT("PreviewBindingLog")));
+
+	ChildSlot
+	[
+		SNew(SHorizontalBox)
+		+ SHorizontalBox::Slot()
+		.AutoWidth()
+		.VAlign(VAlign_Fill)
+		[
+			SourcePanel.ToSharedRef()
+		]
+
+		+ SHorizontalBox::Slot()
+		.FillWidth(1.f)
+		.VAlign(VAlign_Fill)
+		[
+			BindingLog.ToSharedRef()
+		]
+	];
+}
+
+}
diff --git a/Engine/Plugins/Editor/ModelViewViewModelPreview/Source/ModelViewViewModelPreview/Private/Widgets/SMVVMPreviewViewmodels.h b/Engine/Plugins/Editor/ModelViewViewModelPreview/Source/ModelViewViewModelPreview/Private/Widgets/SMVVMPreviewViewmodels.h
new file mode 100644
--- /dev/null
+++ b/Engine/Plugins/Editor/ModelViewViewModelPreview/Source/ModelViewViewModelPreview/Private/Widgets/SMVVMPreviewViewmodels.h
@@ -0,0 +1,32 @@
+// Copyright Epic Games, Inc. All Rights Reserved.
+
+#pragma once
+
+#include "Widgets/SCompoundWidget.h"
+
+namespace UE::UMGWidgetPreview
+{
+	class IWidgetPreviewToolkit;
+}
+
+namespace UE::MVVM::Private
+{
+
+class SPreviewSourceView;
+class SPreviewSourceEntry;
+
+/** */
+class SMVVMPreviewViewmodels : public SCompoundWidget
+{
+public:
+	SLATE_BEGIN_ARGS(SMVVMPreviewViewmodels) {}
+	SLATE_END_ARGS()
+
+	void Construct(const FArguments& InArgs, TSharedPtr<UMGWidgetPreview::IWidgetPreviewToolkit> InPreviewEditor);
+
+private:
+	TSharedPtr<SWidget> SourcePanel;
+	TSharedPtr<SWidget> BindingLog;
+};
+
+}
diff --git a/Engine/Plugins/Runtime/ModelViewViewModel/Source/ModelViewViewModel/Public/Bindings/MVVMCompiledBindingLibrary.h b/Engine/Plugins/Runtime/ModelViewViewModel/Source/ModelViewViewModel/Public/Bindings/MVVMCompiledBindingLibrary.h
--- a/Engine/Plugins/Runtime/ModelViewViewModel/Source/ModelViewViewModel/Public/Bindings/MVVMCompiledBindingLibrary.h
+++ b/Engine/Plugins/Runtime/ModelViewViewModel/Source/ModelViewViewModel/Public/Bindings/MVVMCompiledBindingLibrary.h
@@ -92,6 +92,11 @@ struct FMVVMVCompiledFieldPath
 	{
 		return Num >= 0;
 	}
+	
+	int16 GetStartIndex() const
+	{
+		return StartIndex;
+	}
 
 private:
 	using IndexType = int16;
