UMG MVVM Widget Preview: You can now dump the current viewmodel data to a json file.

#Synced-CL 49024827

diff --git a/Engine/Plugins/Editor/ModelViewViewModelPreview/Source/ModelViewViewModelPreview/ModelViewViewModelPreview.Build.cs b/Engine/Plugins/Editor/ModelViewViewModelPreview/Source/ModelViewViewModelPreview/ModelViewViewModelPreview.Build.cs
--- a/Engine/Plugins/Editor/ModelViewViewModelPreview/Source/ModelViewViewModelPreview/ModelViewViewModelPreview.Build.cs
+++ b/Engine/Plugins/Editor/ModelViewViewModelPreview/Source/ModelViewViewModelPreview/ModelViewViewModelPreview.Build.cs
@@ -19,7 +19,9 @@ public class ModelViewViewModelPreview : ModuleRules
 				"ApplicationCore",
 				"Core",
 				"CoreUObject",
+				"DesktopPlatform",
 				"FieldNotification",
+				"Json",
 				"InputCore",
 				"MessageLog",
 				"ModelViewViewModel",
diff --git a/Engine/Plugins/Editor/ModelViewViewModelPreview/Source/ModelViewViewModelPreview/Private/MVVMUtils.cpp b/Engine/Plugins/Editor/ModelViewViewModelPreview/Source/ModelViewViewModelPreview/Private/MVVMUtils.cpp
new file mode 100644
--- /dev/null
+++ b/Engine/Plugins/Editor/ModelViewViewModelPreview/Source/ModelViewViewModelPreview/Private/MVVMUtils.cpp
@@ -0,0 +1,140 @@
+// Copyright Epic Games, Inc. All Rights Reserved.
+
+#include "MVVMUtils.h"
+
+#include "Containers/UnrealString.h"
+#include "HAL/PlatformApplicationMisc.h"
+#include "HAL/PlatformMemory.h"
+#include "Internationalization/Text.h"
+#include "Misc/Optional.h"
+#include "Templates/ValueOrError.h"
+#include "Types/MVVMFieldContext.h"
+#include "Types/MVVMFieldVariant.h"
+#include "Types/MVVMObjectVariant.h"
+#include "UObject/Class.h"
+#include "UObject/Object.h"
+#include "UObject/TextProperty.h"
+#include "UObject/UObjectGlobals.h"
+#include "UObject/UnrealType.h"
+
+#define LOCTEXT_NAMESPACE "MVVMUtils"
+
+namespace UE::MVVM::Private::Utils
+{
+	// Returns the FProperty for the return value of the function.
+	const FProperty* GetRuntimeReturnProperty(const UFunction* InFunction)
+	{
+		FProperty* Result = InFunction->GetReturnProperty();
+
+		if (Result == nullptr && InFunction->HasAllFunctionFlags(FUNC_HasOutParms))
+		{
+			for (TFieldIterator<FProperty> It(InFunction); It && (It->PropertyFlags & CPF_Parm); ++It)
+			{
+				if (It->HasAllPropertyFlags(CPF_OutParm) && !It->HasAnyPropertyFlags(CPF_Deprecated | CPF_EditorOnly | CPF_ConstParm | CPF_ReferenceParm))
+				{
+					Result = *It;
+					break;
+				}
+			}
+		}
+
+		return Result;
+	}
+
+	TOptional<FText> GetPropertyValue(UObject* SourceObject, FProperty* InProperty)
+	{
+		if (!InProperty)
+		{
+			return {};
+		}
+
+		if (InProperty->IsA<FArrayProperty>())
+		{
+			return LOCTEXT("SourceValueArray", "[Array]");
+		}
+		else if (FObjectProperty* ObjectProperty = CastField<FObjectProperty>(InProperty))
+		{
+			if (UObject* Value = ObjectProperty->GetObjectPropertyValue_InContainer(SourceObject, 0))
+			{
+				return FText::FromString(Value->GetPathName());
+			}
+			else
+			{
+				return LOCTEXT("SourceValueNull", "[None]");
+			}
+		}
+		// By default will create loctext macros.
+		else if (FTextProperty* TextProperty = CastField<FTextProperty>(InProperty))
+		{
+			FText Value;
+			TextProperty->GetValue_InContainer(SourceObject, &Value);
+
+			return Value;
+		}
+		else
+		{
+			FString Value;
+			InProperty->ExportTextItem_InContainer(Value, SourceObject, nullptr, nullptr, PPF_None);
+
+			return FText::FromString(Value);
+		}
+	}
+
+	TOptional<FText> GetFunctionValue(UObject* SourceObject, UFunction* InFunction)
+	{
+		if (!InFunction)
+		{
+			return {};
+		}
+
+		const FProperty* ReturnProperty = GetRuntimeReturnProperty(InFunction);
+
+		if (!ReturnProperty)
+		{
+			return {};
+		}
+
+		// Bound functions must be const/pure and have no parameters, so only a return value
+		// is necessary for the parameters "struct".
+
+		void* DataPtr = FMemory_Alloca_Aligned(ReturnProperty->GetSize(), ReturnProperty->GetMinAlignment());
+		ReturnProperty->InitializeValue(DataPtr);
+
+		SourceObject->ProcessEvent(InFunction, DataPtr);
+
+		FString Value;
+		ReturnProperty->ExportTextItem_Direct(Value, DataPtr, nullptr, nullptr, PPF_None);
+
+		ReturnProperty->DestroyValue(DataPtr);
+
+		return FText::FromString(Value);
+	}
+
+	TOptional<FText> GetFieldValue(TValueOrError<UE::MVVM::FFieldContext, void>& InFieldContext, bool bInAllowFunction)
+	{
+		if (InFieldContext.HasError())
+		{
+			return {};
+		}
+
+		UObject* SourceObject = InFieldContext.GetValue().GetObjectVariant().GetUObject();
+
+		if (!SourceObject)
+		{
+			return {};
+		}
+
+		if (InFieldContext.GetValue().GetFieldVariant().IsProperty())
+		{
+			return GetPropertyValue(SourceObject, InFieldContext.GetValue().GetFieldVariant().GetProperty());
+		}
+		else if (bInAllowFunction && InFieldContext.GetValue().GetFieldVariant().IsFunction())
+		{
+			return GetFunctionValue(SourceObject, InFieldContext.GetValue().GetFieldVariant().GetFunction());
+		}
+
+		return {};
+	}
+} // namespace UE::MVVM::Private::Utils
+
+#undef LOCTEXT_NAMESPACE
diff --git a/Engine/Plugins/Editor/ModelViewViewModelPreview/Source/ModelViewViewModelPreview/Private/MVVMUtils.h b/Engine/Plugins/Editor/ModelViewViewModelPreview/Source/ModelViewViewModelPreview/Private/MVVMUtils.h
new file mode 100644
--- /dev/null
+++ b/Engine/Plugins/Editor/ModelViewViewModelPreview/Source/ModelViewViewModelPreview/Private/MVVMUtils.h
@@ -0,0 +1,30 @@
+// Copyright Epic Games, Inc. All Rights Reserved.
+
+class FProperty;
+class FText;
+class UFunction;
+class UObject;
+template<typename OptionalType> struct TOptional;
+template<typename ValueType, typename ErrorType> class TValueOrError;
+
+namespace UE::MVVM
+{
+
+struct FFieldContext;
+
+namespace Private::Utils
+{
+	// Return a textual representation of the property's value
+	// Returns an unset optional on failure.
+	TOptional<FText> GetPropertyValue(UObject* SourceObject, FProperty* InProperty);
+
+	// Call the function and return a text representation of its return value.
+	// Returns an unset optional on failure.
+	TOptional<FText> GetFunctionValue(UObject* SourceObject, UFunction* InFunction);
+
+	// Return a textual representation of the field's value
+	// Returns an unset optional on failure.
+	TOptional<FText> GetFieldValue(TValueOrError<UE::MVVM::FFieldContext, void>& InFieldContext, bool bInAllowFunction);
+} // namespace Private::Utils
+
+} // namespace UE::MVVM
\ No newline at end of file
diff --git a/Engine/Plugins/Editor/ModelViewViewModelPreview/Source/ModelViewViewModelPreview/Private/ModelViewViewModelPreviewModule.h b/Engine/Plugins/Editor/ModelViewViewModelPreview/Source/ModelViewViewModelPreview/Private/ModelViewViewModelPreviewModule.h
--- a/Engine/Plugins/Editor/ModelViewViewModelPreview/Source/ModelViewViewModelPreview/Private/ModelViewViewModelPreviewModule.h
+++ b/Engine/Plugins/Editor/ModelViewViewModelPreview/Source/ModelViewViewModelPreview/Private/ModelViewViewModelPreviewModule.h
@@ -4,6 +4,7 @@
 
 #include "Modules/ModuleInterface.h"
 
+#include "Logging/LogMacros.h"
 #include "Templates/SharedPointer.h"
 #include "UObject/NameTypes.h"
 
diff --git a/Engine/Plugins/Editor/ModelViewViewModelPreview/Source/ModelViewViewModelPreview/Private/Widgets/SMVVMPreviewBindingLog.cpp b/Engine/Plugins/Editor/ModelViewViewModelPreview/Source/ModelViewViewModelPreview/Private/Widgets/SMVVMPreviewBindingLog.cpp
--- a/Engine/Plugins/Editor/ModelViewViewModelPreview/Source/ModelViewViewModelPreview/Private/Widgets/SMVVMPreviewBindingLog.cpp
+++ b/Engine/Plugins/Editor/ModelViewViewModelPreview/Source/ModelViewViewModelPreview/Private/Widgets/SMVVMPreviewBindingLog.cpp
@@ -3,9 +3,7 @@
 #include "Widgets/SMVVMPreviewBindingLog.h"
 
 #include "Blueprint/UserWidget.h"
-#include "Containers/UnrealString.h"
 #include "HAL/PlatformApplicationMisc.h"
-#include "HAL/PlatformMemory.h"
 #include "IMessageLogListing.h"
 #include "INotifyFieldValueChanged.h"
 #include "Internationalization/Text.h"
@@ -16,6 +14,7 @@
 #include "ModelViewViewModelPreviewModule.h"
 #include "Modules/ModuleManager.h"
 #include "MVVMSubsystem.h"
+#include "MVVMUtils.h"
 #include "Types/MVVMFieldContext.h"
 #include "UObject/Object.h"
 #include "View/MVVMView.h"
@@ -55,127 +54,6 @@ namespace BindingLogUtils
 		return NewSources;
 	}
 	
-	// Returns the FProperty for the return value of the function.
-	const FProperty* GetRuntimeReturnProperty(const UFunction* InFunction)
-	{
-		FProperty* Result = InFunction->GetReturnProperty();
-
-		if (Result == nullptr && InFunction->HasAllFunctionFlags(FUNC_HasOutParms))
-		{
-			for (TFieldIterator<FProperty> It(InFunction); It && (It->PropertyFlags & CPF_Parm); ++It)
-			{
-				if (It->HasAllPropertyFlags(CPF_OutParm) && !It->HasAnyPropertyFlags(CPF_Deprecated | CPF_EditorOnly | CPF_ConstParm | CPF_ReferenceParm))
-				{
-					Result = *It;
-					break;
-				}
-			}
-		}
-
-		return Result;
-	}
-
-	// Interogate the property and add its value as a message token.
-	// Returns an unset optional on failure.
-	TOptional<FText> GetPropertyValue(UObject* SourceObject, FProperty* InProperty)
-	{
-		if (!InProperty)
-		{
-			return {};
-		}
-
-		if (InProperty->IsA<FArrayProperty>())
-		{
-			return LOCTEXT("SourceValueArray", "[Array]");
-		}
-		else if (FObjectProperty* ObjectProperty = CastField<FObjectProperty>(InProperty))
-		{
-			if (UObject* Value = ObjectProperty->GetObjectPropertyValue_InContainer(SourceObject, 0))
-			{
-				return FText::FromString(Value->GetPathName());
-			}
-			else
-			{
-				return LOCTEXT("SourceValueNull", "[None]");
-			}
-		}
-		// By default will create loctext macros.
-		else if (FTextProperty* TextProperty = CastField<FTextProperty>(InProperty))
-		{
-			FText Value;
-			TextProperty->GetValue_InContainer(SourceObject, &Value);
-
-			return Value;
-		}
-		else
-		{
-			FString Value;
-			InProperty->ExportTextItem_InContainer(Value, SourceObject, nullptr, nullptr, PPF_None);
-
-			return FText::FromString(Value);
-		}
-	}
-
-	// Call the function and add its value as a message token.
-	// Returns an unset optional on failure.
-	TOptional<FText> GetFunctionValue(UObject* SourceObject, UFunction* InFunction)
-	{
-		if (!InFunction)
-		{
-			return {};
-		}
-
-		const FProperty* ReturnProperty = GetRuntimeReturnProperty(InFunction);
-
-		if (!ReturnProperty)
-		{
-			return {};
-		}
-
-		// Bound functions must be const/pure and have no parameters, so only a return value
-		// is necessary for the parameters "struct".
-
-		void* DataPtr = FMemory_Alloca_Aligned(ReturnProperty->GetSize(), ReturnProperty->GetMinAlignment());
-		ReturnProperty->InitializeValue(DataPtr);
-
-		SourceObject->ProcessEvent(InFunction, DataPtr);
-
-		FString Value;
-		ReturnProperty->ExportTextItem_Direct(Value, DataPtr, nullptr, nullptr, PPF_None);
-
-		ReturnProperty->DestroyValue(DataPtr);
-
-		return FText::FromString(Value);
-	}
-
-	// Adds the context's value as a token to the message.
-	// Returns an unset optional on failure.
-	TOptional<FText> GetFieldValue(TValueOrError<UE::MVVM::FFieldContext, void>& InFieldContext, bool bInAllowFunction)
-	{
-		if (InFieldContext.HasError())
-		{
-			return {};
-		}
-
-		UObject* SourceObject = InFieldContext.GetValue().GetObjectVariant().GetUObject();
-
-		if (!SourceObject)
-		{
-			return {};
-		}
-
-		if (InFieldContext.GetValue().GetFieldVariant().IsProperty())
-		{
-			return GetPropertyValue(SourceObject, InFieldContext.GetValue().GetFieldVariant().GetProperty());
-		}
-		else if (bInAllowFunction && InFieldContext.GetValue().GetFieldVariant().IsFunction())
-		{
-			return GetFunctionValue(SourceObject, InFieldContext.GetValue().GetFieldVariant().GetFunction());
-		}
-
-		return {};
-	}
-
 	// Finds the first valid source based on the given source mask.
 	UObject* GetFirstValidSourceByMask(const UMVVMView* InView, uint64 InSourceMask)
 	{
@@ -535,7 +413,7 @@ void SMVVMPreviewBindingLog::HandleOnFieldValueChanged(UObject* InSourceObject,
 
 	if (FProperty* Property = Field.GetProperty())
 	{
-		if (TOptional<FText> Value = BindingLogUtils::GetPropertyValue(InSourceObject, Property))
+		if (TOptional<FText> Value = Utils::GetPropertyValue(InSourceObject, Property))
 		{
 			AddCopyToken(SourceValueLog, Value.GetValue());
 		}
@@ -547,7 +425,7 @@ void SMVVMPreviewBindingLog::HandleOnFieldValueChanged(UObject* InSourceObject,
 	}
 	else if (UFunction* Function = Field.GetFunction())
 	{
-		if (TOptional<FText> Value = BindingLogUtils::GetFunctionValue(InSourceObject, Function))
+		if (TOptional<FText> Value = Utils::GetFunctionValue(InSourceObject, Function))
 		{
 			AddCopyToken(SourceValueLog, Value.GetValue());
 		}
@@ -681,7 +559,7 @@ void SMVVMPreviewBindingLog::HandleOnLibraryBindingExecuted(const FDebugging::FV
 			}
 		}
 
-		if (TOptional<FText> Value = BindingLogUtils::GetFieldValue(DestinationFieldContext, /* Allow Function */ false))
+		if (TOptional<FText> Value = Utils::GetFieldValue(DestinationFieldContext, /* Allow Function */ false))
 		{
 			AddCopyToken(DestinationValueLog, Value.GetValue());
 		}
diff --git a/Engine/Plugins/Editor/ModelViewViewModelPreview/Source/ModelViewViewModelPreview/Private/Widgets/SMVVMPreviewSourcePanel.cpp b/Engine/Plugins/Editor/ModelViewViewModelPreview/Source/ModelViewViewModelPreview/Private/Widgets/SMVVMPreviewSourcePanel.cpp
--- a/Engine/Plugins/Editor/ModelViewViewModelPreview/Source/ModelViewViewModelPreview/Private/Widgets/SMVVMPreviewSourcePanel.cpp
+++ b/Engine/Plugins/Editor/ModelViewViewModelPreview/Source/ModelViewViewModelPreview/Private/Widgets/SMVVMPreviewSourcePanel.cpp
@@ -2,9 +2,23 @@
 
 #include "Widgets/SMVVMPreviewSourcePanel.h"
 
+#include "DesktopPlatformModule.h"
+#include "Dom/JsonObject.h"
+#include "Dom/JsonValue.h"
+#include "Framework/Application/SlateApplication.h"
+#include "IDesktopPlatform.h"
+#include "IFieldNotificationClassDescriptor.h"
 #include "IWidgetPreviewToolkit.h"
+#include "Logging/LogMacros.h"
+#include "Misc/FileHelper.h"
+#include "Misc/MessageDialog.h"
+#include "MVVMMessageLog.h"
 #include "MVVMSubsystem.h"
+#include "MVVMUtils.h"
+#include "Serialization/JsonSerializer.h"
+#include "Serialization/JsonWriter.h"
 #include "Styling/AppStyle.h"
+#include "UObject/Object.h"
 #include "View/MVVMView.h"
 #include "View/MVVMViewClass.h"
 #include "WidgetPreview.h"
@@ -78,6 +92,7 @@ namespace UE::MVVM::Private
 				.SelectionMode(ESelectionMode::Single)
 				.OnGenerateRow(this, &SPreviewSourcePanel::GenerateWidget)
 				.OnSelectionChanged(this, &SPreviewSourcePanel::HandleSourceSelectionChanged)
+				.OnContextMenuOpening(this, &SPreviewSourcePanel::OnContextMenuOpened)
 				.HeaderRow(
 					SNew(SHeaderRow)
 					+ SHeaderRow::Column("Sources")
@@ -90,7 +105,6 @@ namespace UE::MVVM::Private
 		HandlePreviewWidgetChanged(EWidgetPreviewWidgetChangeType::Reinstanced);
 	}
 
-
 	SPreviewSourcePanel::~SPreviewSourcePanel()
 	{
 		if (TSharedPtr<UMGWidgetPreview::IWidgetPreviewToolkit> PreviewEditor = WeakPreviewEditor.Pin())
@@ -104,27 +118,22 @@ namespace UE::MVVM::Private
 		}
 	}
 
-
 	void SPreviewSourcePanel::HandlePreviewWidgetChanged(const EWidgetPreviewWidgetChangeType InChangeType)
 	{
 		SourceList.Reset();
 		WeakView.Reset();
 
-		// Only respond to instantiation changes
-		// if (InChangeType == EWidgetPreviewWidgetChangeType::Reinstanced)
+		if (TSharedPtr<UMGWidgetPreview::IWidgetPreviewToolkit> PreviewEditor = WeakPreviewEditor.Pin())
 		{
-			if (TSharedPtr<UMGWidgetPreview::IWidgetPreviewToolkit> PreviewEditor = WeakPreviewEditor.Pin())
+			if (const UUserWidget* NewWidget = PreviewEditor->GetPreview() ? PreviewEditor->GetPreview()->GetWidgetInstance() : nullptr)
 			{
-				if (const UUserWidget* NewWidget = PreviewEditor->GetPreview() ? PreviewEditor->GetPreview()->GetWidgetInstance() : nullptr)
+				if (UMVVMView* View = UMVVMSubsystem::GetViewFromUserWidget(NewWidget))
 				{
-					if (UMVVMView* View = UMVVMSubsystem::GetViewFromUserWidget(NewWidget))
+					WeakView = View;
+					for (const FMVVMView_Source& ViewSource : View->GetSources())
 					{
-						WeakView = View;
-						for (const FMVVMView_Source& ViewSource : View->GetSources())
-						{
-							FName SourceName = View->GetViewClass()->GetSource(ViewSource.ClassKey).GetName();
-							SourceList.Emplace(MakeShared<Private::SPreviewSourceEntry>(ViewSource.Source, SourceName));
-						}
+						FName SourceName = View->GetViewClass()->GetSource(ViewSource.ClassKey).GetName();
+						SourceList.Emplace(MakeShared<Private::SPreviewSourceEntry>(ViewSource.Source, SourceName));
 					}
 				}
 			}
@@ -136,7 +145,6 @@ namespace UE::MVVM::Private
 		}
 	}
 
-
 	void SPreviewSourcePanel::HandleSelectedObjectChanged(const TConstArrayView<TWeakObjectPtr<UObject>> InSelectedObjects)
 	{
 		if (bIsSelectingListItem)
@@ -171,7 +179,6 @@ namespace UE::MVVM::Private
 		}
 	}
 
-
 	void SPreviewSourcePanel::HandleSourceSelectionChanged(TSharedPtr<Private::SPreviewSourceEntry> Entry, ESelectInfo::Type SelectionType) const
 	{
 		if (bIsSelectingListItem)
@@ -207,14 +214,14 @@ namespace UE::MVVM::Private
 	}
 	#endif
 
-
 	TSharedRef<ITableRow> SPreviewSourcePanel::GenerateWidget(TSharedPtr<Private::SPreviewSourceEntry> Entry, const TSharedRef<STableViewBase>& OwnerTable) const
 	{
-		typedef STableRow<TSharedPtr<Private::SPreviewSourceEntry>> RowType;
+		using RowType = STableRow<TSharedPtr<Private::SPreviewSourceEntry>>;
 
 		const TSharedRef<SWidget> FieldIcon = Entry->GetClass() ? SNew(UE::PropertyViewer::SFieldIcon, Entry->GetClass()) : SNullWidget::NullWidget;
 
 		TSharedRef<RowType> NewRow = SNew(RowType, OwnerTable);
+
 		NewRow->SetContent(SNew(SHorizontalBox)
 			+ SHorizontalBox::Slot()
 			.AutoWidth()
@@ -232,6 +239,165 @@ namespace UE::MVVM::Private
 
 		return NewRow;
 	}
+
+	TSharedPtr<SWidget> SPreviewSourcePanel::OnContextMenuOpened()
+	{
+		if (SourceListView.IsValid())
+		{
+			TArray<TSharedPtr<Private::SPreviewSourceEntry>> SelectedItems = SourceListView->GetSelectedItems();
+
+			if (SelectedItems.Num() == 1)
+			{
+				constexpr bool bInShouldCloseWindowAfterMenuSelection = true;
+				FMenuBuilder MenuBuilder(bInShouldCloseWindowAfterMenuSelection, nullptr);
+
+				MenuBuilder.BeginSection(TEXT("Viewmodel"), LOCTEXT("Viewmodel", "Viewmodel"));
+				{
+					FMenuEntryParams Params;
+					Params.LabelOverride = LOCTEXT("ViewmodelDataDump", "Dump Data to JSON");
+					Params.UserInterfaceActionType = EUserInterfaceActionType::Button;
+					Params.DirectActions.ExecuteAction.BindSP(
+						this,
+						&SPreviewSourcePanel::ExecuteDumpData,
+						SelectedItems[0]
+					);
+					
+					MenuBuilder.AddMenuEntry(Params);
+				}
+				MenuBuilder.EndSection();
+
+				return MenuBuilder.MakeWidget();
+			}
+		}
+
+		return SNullWidget::NullWidget;
+	}
+
+	void SPreviewSourcePanel::ExecuteDumpData(TSharedPtr<Private::SPreviewSourceEntry> InEntry)
+	{
+		if (!InEntry.IsValid())
+		{
+			return;
+		}
+		
+		UObject* Entry = InEntry->GetInstance();
+
+		if (!Entry)
+		{
+			return;
+		}
+
+		if (!Entry->Implements<UNotifyFieldValueChanged>())
+		{
+			UE_LOG(LogMVVM, Error, TEXT("Object selected for dump does not implement INotifyFieldValueChanged. [%s]"), *Entry->GetPathName());
+			return;
+		}
+
+		IDesktopPlatform* DesktopPlatform = FDesktopPlatformModule::TryGet();
+
+		if (!DesktopPlatform)
+		{
+			UE_LOG(LogMVVM, Error, TEXT("Unable to retrieve desktop platform instance."));
+			return;
+		}
+
+		static const TArray<TPair<TCHAR, TCHAR>> ValidCharacterRanges = {
+			{'A', 'Z'},
+			{'a', 'z'},
+			{'0', '9'},
+			{' ', ' '},
+			{'-', '-'},
+			{'_', '_'}
+		};
+
+		FString Filename = Entry->GetPathName();
+
+		for (int32 Index = 0; Index < Filename.Len(); ++Index)
+		{
+			bool bIsValidChar = false;
+
+			for (const TPair<TCHAR, TCHAR>& CharRange : ValidCharacterRanges)
+			{
+				if (Filename[Index] >= CharRange.Key && Filename[Index] <= CharRange.Value)
+				{
+					bIsValidChar = true;
+					break;
+				}
+			}
+
+			if (!bIsValidChar)
+			{
+				Filename[Index] = TEXT('_');
+			}
+		}
+
+		TArray<FString> Filenames;
+		Filenames.Reserve(1);
+
+		const bool bSuccess = DesktopPlatform->SaveFileDialog(
+			FSlateApplication::Get().FindBestParentWindowHandleForDialogs(SharedThis(this)),
+			LOCTEXT("DumpDataTitle", "Export Viewmodel to...").ToString(),
+			FPaths::ProjectSavedDir(),
+			Filename,
+			TEXT("json (*.json)|*.json|"),
+			EFileDialogFlags::None,
+			Filenames
+		);
+
+		// Nothing picked
+		if (!bSuccess || Filenames.IsEmpty())
+		{
+			return;
+		}
+
+		TSharedRef<FJsonObject> JsonRoot = MakeShared<FJsonObject>();
+		TSharedRef<FJsonObject> Functions = MakeShared<FJsonObject>();
+		TSharedRef<FJsonObject> Properties = MakeShared<FJsonObject>();
+		
+		JsonRoot->SetField(TEXT("Functions"), MakeShared<FJsonValueObject>(Functions));
+		JsonRoot->SetField(TEXT("Properties"), MakeShared<FJsonValueObject>(Properties));
+
+		INotifyFieldValueChanged* FieldNotifier = Cast<INotifyFieldValueChanged>(Entry);
+		const UE::FieldNotification::IClassDescriptor& ClassDescriptor = FieldNotifier->GetFieldNotificationDescriptor();
+		
+		ClassDescriptor.ForEachField(
+			Entry->GetClass(),
+			[this, Entry, Functions, Properties](UE::FieldNotification::FFieldId InFieldId) -> bool
+			{
+				UE::FieldNotification::FFieldVariant Field = InFieldId.ToVariant(Entry);
+
+				if (UFunction* Function = Field.GetFunction())
+				{
+					Functions->SetStringField(
+						Function->GetName(),
+						Utils::GetFunctionValue(Entry, Function).Get(FText::GetEmpty()).ToString()
+					);
+				}
+				else if (FProperty* Property = Field.GetProperty())
+				{
+					Properties->SetStringField(
+						Property->GetName(),
+						Utils::GetPropertyValue(Entry, Property).Get(FText::GetEmpty()).ToString()
+					);
+				}
+
+				return true;
+			}
+		);
+
+		FString SerializedString;
+		const TSharedRef<TJsonWriter<>> Writer = TJsonWriterFactory<>::Create(&SerializedString);
+		FJsonSerializer::Serialize(JsonRoot, Writer);
+
+		if (FFileHelper::SaveStringToFile(SerializedString, *Filenames[0]))
+		{
+			FMessageDialog::Open(EAppMsgType::Ok, LOCTEXT("DumpSuccess", "Viewmodel data dumped to JSON."));
+		}
+		else
+		{
+			FMessageDialog::Open(EAppMsgType::Ok, LOCTEXT("DumpFailure", "Failed to dump Viewmodel data to JSON."));
+		}
+	}
 }
 
 #undef LOCTEXT_NAMESPACE
diff --git a/Engine/Plugins/Editor/ModelViewViewModelPreview/Source/ModelViewViewModelPreview/Private/Widgets/SMVVMPreviewSourcePanel.h b/Engine/Plugins/Editor/ModelViewViewModelPreview/Source/ModelViewViewModelPreview/Private/Widgets/SMVVMPreviewSourcePanel.h
--- a/Engine/Plugins/Editor/ModelViewViewModelPreview/Source/ModelViewViewModelPreview/Private/Widgets/SMVVMPreviewSourcePanel.h
+++ b/Engine/Plugins/Editor/ModelViewViewModelPreview/Source/ModelViewViewModelPreview/Private/Widgets/SMVVMPreviewSourcePanel.h
@@ -44,6 +44,8 @@ namespace UE::MVVM::Private
 		void HandleViewChanged(const FDebugging::FView&, const FDebugging::FViewSourceValueArgs&);
 	#endif
 		TSharedRef<ITableRow> GenerateWidget(TSharedPtr<Private::SPreviewSourceEntry> Entry, const TSharedRef<STableViewBase>& OwnerTable) const;
+		TSharedPtr<SWidget> OnContextMenuOpened();
+		void ExecuteDumpData(TSharedPtr<Private::SPreviewSourceEntry> InEntry);
 
 		TWeakObjectPtr<UMVVMView> WeakView;
 		TWeakPtr<UMGWidgetPreview::IWidgetPreviewToolkit> WeakPreviewEditor;
