UMG MVVM Previewer: There's now an option to default construct missing viewmodels when the preview is opened.

#14095
#Synced-CL 49521844

diff --git a/Engine/Plugins/Editor/ModelViewViewModelPreview/Source/ModelViewViewModelPreview/Private/Widgets/SMVVMPreviewSourcePanel.cpp b/Engine/Plugins/Editor/ModelViewViewModelPreview/Source/ModelViewViewModelPreview/Private/Widgets/SMVVMPreviewSourcePanel.cpp
--- a/Engine/Plugins/Editor/ModelViewViewModelPreview/Source/ModelViewViewModelPreview/Private/Widgets/SMVVMPreviewSourcePanel.cpp
+++ b/Engine/Plugins/Editor/ModelViewViewModelPreview/Source/ModelViewViewModelPreview/Private/Widgets/SMVVMPreviewSourcePanel.cpp
@@ -2,6 +2,7 @@
 
 #include "Widgets/SMVVMPreviewSourcePanel.h"
 
+#include "Blueprint/UserWidget.h"
 #include "DesktopPlatformModule.h"
 #include "Dom/JsonObject.h"
 #include "Dom/JsonValue.h"
@@ -125,13 +126,34 @@ namespace UE::MVVM::Private
 
 		if (TSharedPtr<UMGWidgetPreview::IWidgetPreviewToolkit> PreviewEditor = WeakPreviewEditor.Pin())
 		{
-			if (const UUserWidget* NewWidget = PreviewEditor->GetPreview() ? PreviewEditor->GetPreview()->GetWidgetInstance() : nullptr)
+			if (UUserWidget* NewWidget = PreviewEditor->GetPreview() ? PreviewEditor->GetPreview()->GetWidgetInstance() : nullptr)
 			{
 				if (UMVVMView* View = UMVVMSubsystem::GetViewFromUserWidget(NewWidget))
 				{
 					WeakView = View;
+
+					const bool bShouldDefaultConstructViewmodels = PreviewEditor->GetPreview()->ShouldDefaultConstructViewmodels();
+
 					for (const FMVVMView_Source& ViewSource : View->GetSources())
 					{
+						/**
+						 * If we don't have a viewmodel, create one now. 
+						 * If the Viewmodel creation mode is manual, global or resolver based the viewmodel won't exist in editor
+						 * but we still want to show something in the preview.
+						 */
+						if (bShouldDefaultConstructViewmodels && !ViewSource.Source)
+						{
+							const FMVVMViewClass_Source& ClassSource = View->GetViewClass()->GetSource(ViewSource.ClassKey);
+							
+							if (UClass* SourceClass = ClassSource.GetSourceClass())
+							{
+								if (UObject* NewSource = NewObject<UObject>(NewWidget, SourceClass, NAME_None, RF_Transient))
+								{
+									View->SetViewModelByClass(NewSource);
+								}
+							}
+						}
+
 						FName SourceName = View->GetViewClass()->GetSource(ViewSource.ClassKey).GetName();
 						SourceList.Emplace(MakeShared<Private::SPreviewSourceEntry>(ViewSource.Source, SourceName));
 					}
diff --git a/Engine/Plugins/Editor/UMGWidgetPreview/Source/UMGWidgetPreview/Private/WidgetPreview.cpp b/Engine/Plugins/Editor/UMGWidgetPreview/Source/UMGWidgetPreview/Private/WidgetPreview.cpp
--- a/Engine/Plugins/Editor/UMGWidgetPreview/Source/UMGWidgetPreview/Private/WidgetPreview.cpp
+++ b/Engine/Plugins/Editor/UMGWidgetPreview/Source/UMGWidgetPreview/Private/WidgetPreview.cpp
@@ -115,6 +115,7 @@ const UWidgetPreview* FPreviewableWidgetVariant::AsWidgetPreview() const
 UWidgetPreview::UWidgetPreview(const FObjectInitializer& ObjectInitializer)
 	: Super(ObjectInitializer)
 	, bShouldOverrideWidgetSize(false)
+	, bDefaultConstructViewmodels(false)
 {
 	OverriddenWidgetSize = FVector2D(100,100);
 }
@@ -405,6 +406,16 @@ void UWidgetPreview::SetOverriddenWidgetSize(FVector2D InWidgetSize)
 	
 }
 
+bool UWidgetPreview::ShouldDefaultConstructViewmodels() const
+{
+	return bDefaultConstructViewmodels;
+}
+
+void UWidgetPreview::SetDefaultConstructViewmodels(bool bInDefaultConstruct)
+{
+	bDefaultConstructViewmodels = bInDefaultConstruct;
+}
+
 void UWidgetPreview::PostLoad()
 {
 	UObject::PostLoad();
diff --git a/Engine/Plugins/Editor/UMGWidgetPreview/Source/UMGWidgetPreview/Public/WidgetPreview.h b/Engine/Plugins/Editor/UMGWidgetPreview/Source/UMGWidgetPreview/Public/WidgetPreview.h
--- a/Engine/Plugins/Editor/UMGWidgetPreview/Source/UMGWidgetPreview/Public/WidgetPreview.h
+++ b/Engine/Plugins/Editor/UMGWidgetPreview/Source/UMGWidgetPreview/Public/WidgetPreview.h
@@ -120,6 +120,9 @@ public:
 	UE_API const FVector2D GetOverriddenWidgetSize() const;
 	UE_API void SetOverriddenWidgetSize(FVector2D InWidgetSize);
 
+	UE_API bool ShouldDefaultConstructViewmodels() const;
+	UE_API void SetDefaultConstructViewmodels(bool bInDefaultConstruct);
+
 protected:
 	UE_API virtual void PostLoad() override;
 
@@ -155,6 +158,10 @@ private:
 	UPROPERTY(EditAnywhere, BlueprintReadWrite, Getter, Setter, Category = "Widget", DisplayName = "Widget Size", meta = (AllowPrivateAccess = "true", EditCondition="bShouldOverrideWidgetSize"))
 	FVector2D OverriddenWidgetSize;
 
+	/** When this preview is opened in the previewer, automatically construct any missing viewmodels. */
+	UPROPERTY(EditAnywhere, BlueprintReadWrite, Getter = ShouldDefaultConstructViewmodels, Setter = SetDefaultConstructViewmodels, Category = "Widget", meta = (AllowPrivateAccess = "true"))
+	bool bDefaultConstructViewmodels;
+
 	UPROPERTY(DuplicateTransient)
 	TObjectPtr<UUserWidget> WidgetInstance;
 
